Renamed "TaskGraph" module to "ProfileVisualizer"

diff --git a/Engine/Plugins/Experimental/ChaosClothAsset/Source/ChaosClothAssetEngine/ChaosClothAssetEngine.Build.cs b/Engine/Plugins/Experimental/ChaosClothAsset/Source/ChaosClothAssetEngine/ChaosClothAssetEngine.Build.cs
--- a/Engine/Plugins/Experimental/ChaosClothAsset/Source/ChaosClothAssetEngine/ChaosClothAssetEngine.Build.cs
+++ b/Engine/Plugins/Experimental/ChaosClothAsset/Source/ChaosClothAssetEngine/ChaosClothAssetEngine.Build.cs
@@ -19,7 +19,7 @@ public class ChaosClothAssetEngine : ModuleRules
 				"Chaos",
 				"ChaosClothAsset",
 				"ChaosCloth",
-				"TaskGraph",
+				"ProfileVisualizer",
 				"ClothingSystemRuntimeInterface",
 				"DataflowEngine"
 			}
diff --git a/Engine/Source/Developer/ProfileVisualizer/Private/ProfileVisualizerModule.cpp b/Engine/Source/Developer/ProfileVisualizer/Private/ProfileVisualizerModule.cpp
new file mode 100644
--- /dev/null
+++ b/Engine/Source/Developer/ProfileVisualizer/Private/ProfileVisualizerModule.cpp
@@ -0,0 +1,203 @@
+// Copyright Epic Games, Inc. All Rights Reserved.
+
+#include "ProfileVisualizerModule.h"
+#include "Stats/Stats.h"
+#include "Modules/ModuleManager.h"
+#include "Widgets/DeclarativeSyntaxSupport.h"
+#include "Application/SlateWindowHelper.h"
+#include "Framework/Docking/TabManager.h"
+#include "Tickable.h"
+#include "SProfileVisualizer.h"
+#include "Widgets/Docking/SDockTab.h"
+#include "ProfileVisualizerStyle.h"
+#if WITH_EDITOR
+#include "WorkspaceMenuStructure.h"
+#include "WorkspaceMenuStructureModule.h"
+#include "Styling/AppStyle.h"
+#endif //#if WITH_EDITOR
+
+/**
+ * Creates Visualizer using Visualizer profile data format
+ *
+ * @param ProfileData Visualizer data
+ * @return Visualizer window
+ */
+
+FName ProfileVisualizerTabId("VisualizerSpawnPoint");
+TSharedRef<SDockTab> MakeProfileVisualizerWindow( TSharedPtr< FVisualizerEvent > ProfileData, const FText& WindowTitle, const FText& ProfilerType, const FText& HeaderMessageText = FText::GetEmpty(), const FLinearColor& HeaderMessageTextColor = FLinearColor::White, const bool InsertTab = true )
+{
+	TSharedRef<SDockTab> ProfileVisualizer = SNew(SDockTab)
+												.Label(WindowTitle)
+												.TabRole(ETabRole::NomadTab)
+												[
+													SNew(SProfileVisualizer)
+													.ProfileData(ProfileData)
+													.ProfilerType(ProfilerType)
+													.HeaderMessageText(HeaderMessageText)
+													.HeaderMessageTextColor(HeaderMessageTextColor)
+												];
+
+	if (InsertTab)
+	{
+		FGlobalTabmanager::Get()->InsertNewDocumentTab
+		(
+			ProfileVisualizerTabId, FTabManager::ESearchPreference::RequireClosedTab,
+			ProfileVisualizer
+		);
+	}
+
+	return ProfileVisualizer;
+}
+
+/** Helper class that counts down when to unpause and stop movie. */
+class FDelayedVisualizerSpawner : public FTickableGameObject
+{
+public:
+
+	struct FPendingWindow
+	{
+		FText Title;
+		FText Type;
+		TSharedPtr< FVisualizerEvent > ProfileData;
+
+		FPendingWindow( TSharedPtr<FVisualizerEvent> InData, const FText& InTitle, const FText& InType )
+			: Title( InTitle )
+			, Type( InType )
+			, ProfileData( InData )
+		{}
+	};
+
+	FDelayedVisualizerSpawner( )
+	{
+	}
+
+	virtual ~FDelayedVisualizerSpawner()
+	{
+	}
+
+	void AddPendingData( TSharedPtr<FVisualizerEvent> InProfileData, const FText& InTitle, const FText& InType )
+	{
+		DataLock.Lock();
+
+		VisualizerDataToSpawn.Add( MakeShareable( new FPendingWindow( InProfileData, InTitle, InType ) )  );
+
+		DataLock.Unlock();
+	}
+
+	// FTickableGameObject interface
+
+	void Tick(float DeltaTime) override
+	{
+		DataLock.Lock();
+
+		for ( int32 Index = 0; Index < VisualizerDataToSpawn.Num(); Index++ )
+		{
+			TSharedPtr<FPendingWindow> Window =  VisualizerDataToSpawn[ Index ];
+			MakeProfileVisualizerWindow( Window->ProfileData, Window->Title, Window->Type ); 
+		}
+
+		VisualizerDataToSpawn.Empty();
+
+		DataLock.Unlock();
+	}
+
+	/** We should call Tick on this object */
+	virtual ETickableTickType GetTickableTickType() const override
+	{
+		return ETickableTickType::Always;
+	}
+
+	/** Need this to be ticked when paused (that is the point!) */
+	virtual bool IsTickableWhenPaused() const override
+	{
+		return true;
+	}
+
+	virtual TStatId GetStatId() const override
+	{
+		RETURN_QUICK_DECLARE_CYCLE_STAT(FDelayedVisualizerSpawner, STATGROUP_Tickables);
+	}
+
+private:
+
+	FCriticalSection DataLock;
+
+	/** List of profile data sets to spawn visualizers for */
+	TArray< TSharedPtr<FPendingWindow> > VisualizerDataToSpawn;
+};
+
+static TSharedPtr< FDelayedVisualizerSpawner > GDelayedVisualizerSpawner;
+
+void InitProfileVisualizer()
+{
+	FProfileVisualizerStyle::Initialize();
+	if( GDelayedVisualizerSpawner.IsValid() == false )
+	{
+		GDelayedVisualizerSpawner = MakeShareable( new FDelayedVisualizerSpawner() );
+	}
+}
+
+void ShutdownProfileVisualizer()
+{
+	FProfileVisualizerStyle::Shutdown();
+	GDelayedVisualizerSpawner.Reset();
+}
+
+static bool GHasRegisteredVisualizerLayout = false;
+
+void DisplayProfileVisualizer(TSharedPtr< FVisualizerEvent > InProfileData, const TCHAR* InProfilerType, const FText& HeaderMessageText = FText::GetEmpty(), const FLinearColor& HeaderMessageTextColor = FLinearColor::White)
+{
+	check( IsInGameThread() );
+
+	FFormatNamedArguments Args;
+	Args.Add( TEXT("ProfilerType"), FText::FromString( InProfilerType ) );
+
+	const FText WindowTitle = FText::Format( NSLOCTEXT("ProfileVisualizer", "WindowTitle", "{ProfilerType} Visualizer"), Args );
+	const FText ProfilerType = FText::Format( NSLOCTEXT("ProfileVisualizer", "ProfilerType", "{ProfilerType} Profile"), Args );
+
+	MakeProfileVisualizerWindow( InProfileData, WindowTitle, ProfilerType, HeaderMessageText, HeaderMessageTextColor );
+	
+}
+
+/**
+ * Module for profile visualizer.
+ */
+class FProfileVisualizerModule : public IProfileVisualizerModule
+{
+public:
+	virtual void StartupModule() override
+	{
+		::InitProfileVisualizer();
+
+#if WITH_EDITOR
+		FGlobalTabmanager::Get()->RegisterNomadTabSpawner(ProfileVisualizerTabId, FOnSpawnTab::CreateRaw(this, &FProfileVisualizerModule::SpawnProfileVisualizerTab))
+			.SetDisplayName(NSLOCTEXT("ProfileVisualizerModule", "TabTitle", "Profile Data Visualizer"))
+			.SetTooltipText(NSLOCTEXT("ProfileVisualizerModule", "TooltipText", "Open the Profile Data Visualizer tab."))
+			.SetGroup(WorkspaceMenu::GetMenuStructure().GetDeveloperToolsProfilingCategory())
+			.SetIcon(FSlateIcon(FAppStyle::Get().GetStyleSetName(), "Profiler.Tab"));
+#endif
+			
+	}
+	virtual void ShutdownModule() override
+	{
+#if WITH_EDITOR
+		FGlobalTabmanager::Get()->UnregisterNomadTabSpawner(ProfileVisualizerTabId);
+#endif		
+		::ShutdownProfileVisualizer();
+	}
+	virtual void DisplayProfileVisualizer(TSharedPtr< FVisualizerEvent > InProfileData, const TCHAR* InProfilerType, const FText& HeaderMessageText, const FLinearColor& HeaderMessageTextColor) override
+	{
+#if	WITH_EDITOR
+		::DisplayProfileVisualizer( InProfileData, InProfilerType, HeaderMessageText, HeaderMessageTextColor );
+#endif // WITH_EDITOR
+	}
+
+	TSharedRef<SDockTab> SpawnProfileVisualizerTab(const FSpawnTabArgs& Args)
+	{
+		TSharedPtr< FVisualizerEvent > InProfileData(new FVisualizerEvent(0., 0., 0., 0, "Dummy"));
+		
+		return MakeProfileVisualizerWindow(InProfileData, FText::GetEmpty(), FText::GetEmpty(), FText::GetEmpty(), FLinearColor::White, false);
+	}
+
+};
+IMPLEMENT_MODULE(FProfileVisualizerModule, ProfileVisualizer);
diff --git a/Engine/Source/Developer/ProfileVisualizer/Private/ProfileVisualizerStyle.cpp b/Engine/Source/Developer/ProfileVisualizer/Private/ProfileVisualizerStyle.cpp
new file mode 100644
--- /dev/null
+++ b/Engine/Source/Developer/ProfileVisualizer/Private/ProfileVisualizerStyle.cpp
@@ -0,0 +1,59 @@
+// Copyright Epic Games, Inc. All Rights Reserved.
+
+#include "ProfileVisualizerStyle.h"
+#include "Styling/SlateStyleRegistry.h"
+#include "Styling/SlateStyleMacros.h"
+
+TSharedPtr< FProfileVisualizerStyle::FStyle > FProfileVisualizerStyle::StyleInstance;
+
+FProfileVisualizerStyle::FStyle::FStyle()
+	: FSlateStyleSet("ProfileVisualizerStyle")
+{
+
+}
+
+void FProfileVisualizerStyle::FStyle::Initialize()
+{
+#if WITH_EDITORONLY_DATA
+	StyleInstance->SetContentRoot( FPaths::EngineContentDir() / TEXT("Editor/Slate") );
+	StyleInstance->SetCoreContentRoot(FPaths::EngineContentDir() / TEXT("Slate"));
+
+	StyleInstance->Set( "StatsHeader", new BOX_BRUSH( "Common/TableViewHeader", 4.f/32.f ) );
+	StyleInstance->Set( "ProfileVisualizer.ContentAreaBrush", new BOX_BRUSH( "/Docking/TabContentArea", FMargin(4/16.0f) ) );
+	StyleInstance->Set( "ProfileVisualizer.Background", new BOX_BRUSH( "Common/ProgressBar_Background", FMargin(5.f/12.f) ) );
+	StyleInstance->Set( "ProfileVisualizer.Normal", new CORE_BOX_BRUSH( "Common/ProfileVisualizer_Normal", FMargin(5.f/12.f) ) );
+	StyleInstance->Set( "ProfileVisualizer.Selected", new CORE_BOX_BRUSH( "Common/ProfileVisualizer_Selected", FMargin(5.f/12.f) ) );
+	StyleInstance->Set( "ProfileVisualizer.Mono", new CORE_BOX_BRUSH( "Common/ProfileVisualizer_Mono", FMargin(5.f/12.f) ) );
+	StyleInstance->Set( "ProfileVisualizer.BorderPadding", FVector2D(1,0) );
+	StyleInstance->Set( "ProfileVisualizer.SortUp", new IMAGE_BRUSH( "Common/SortUpArrow", FVector2D(8,4) ) );
+	StyleInstance->Set( "ProfileVisualizer.SortDown", new IMAGE_BRUSH( "Common/SortDownArrow", FVector2D(8,4) ) );
+	StyleInstance->Set( "ProfileVisualizer.Home", new IMAGE_BRUSH( "Icons/Home16x16", FVector2D(16,16) ) );
+	StyleInstance->Set( "ProfileVisualizer.ToParent", new IMAGE_BRUSH( "Icons/ToParent", FVector2D(16,16) ) );
+	StyleInstance->Set( "ProfileVisualizer.ProgressBar.BorderPadding", FVector2D(1,0) );
+	StyleInstance->Set( "ProfileVisualizer.MenuDropdown", new IMAGE_BRUSH( "Common/ComboArrow", FVector2D(8,8) ) );
+
+	FSlateStyleRegistry::RegisterSlateStyle(*StyleInstance.Get());
+#endif
+}
+
+
+void FProfileVisualizerStyle::Initialize() 
+{
+	StyleInstance = MakeShareable( new FProfileVisualizerStyle::FStyle );
+	StyleInstance->Initialize();
+}
+
+void FProfileVisualizerStyle::Shutdown()
+{
+	if (StyleInstance.IsValid())
+	{
+		FSlateStyleRegistry::UnRegisterSlateStyle(*StyleInstance.Get());
+		ensure(StyleInstance.IsUnique());
+		StyleInstance.Reset();
+	}
+}
+
+TSharedPtr< ISlateStyle > FProfileVisualizerStyle::Get()
+{
+	return StyleInstance;
+}
diff --git a/Engine/Source/Developer/ProfileVisualizer/Private/ProfileVisualizerStyle.h b/Engine/Source/Developer/ProfileVisualizer/Private/ProfileVisualizerStyle.h
new file mode 100644
--- /dev/null
+++ b/Engine/Source/Developer/ProfileVisualizer/Private/ProfileVisualizerStyle.h
@@ -0,0 +1,32 @@
+// Copyright Epic Games, Inc. All Rights Reserved.
+
+#pragma once
+
+
+#include "CoreMinimal.h"
+#include "Styling/ISlateStyle.h"
+#include "Styling/SlateStyle.h"
+
+/**
+ * Declares the Editor's visual style.
+ */
+class FProfileVisualizerStyle
+{
+public:
+	
+	class FStyle : public FSlateStyleSet
+	{
+	public:
+		FStyle();
+
+		void Initialize();
+	};
+
+	static void Initialize();
+
+	static void Shutdown();
+
+	static TSharedPtr<class ISlateStyle> Get();
+
+	static TSharedPtr< FProfileVisualizerStyle::FStyle > StyleInstance;
+};
diff --git a/Engine/Source/Developer/ProfileVisualizer/Private/SBarVisualizer.cpp b/Engine/Source/Developer/ProfileVisualizer/Private/SBarVisualizer.cpp
new file mode 100644
--- /dev/null
+++ b/Engine/Source/Developer/ProfileVisualizer/Private/SBarVisualizer.cpp
@@ -0,0 +1,661 @@
+// Copyright Epic Games, Inc. All Rights Reserved.
+
+#include "SBarVisualizer.h"
+#include "SlateOptMacros.h"
+#include "Textures/SlateIcon.h"
+#include "Framework/Commands/UIAction.h"
+#include "Widgets/Images/SImage.h"
+#include "Widgets/Text/STextBlock.h"
+#include "Framework/MultiBox/MultiBoxBuilder.h"
+#include "Widgets/Layout/SScrollBar.h"
+#include "Widgets/Input/SButton.h"
+#include "Widgets/Input/SComboButton.h"
+#include "Widgets/Views/SListView.h"
+#include "Widgets/Input/SSlider.h"
+#include "SGraphBar.h"
+#include "STimeline.h"
+#include "ProfileVisualizerStyle.h"
+
+BEGIN_SLATE_FUNCTION_BUILD_OPTIMIZATION
+void SBarVisualizer::Construct( const FArguments& InArgs )
+{
+	ZoomSliderValue = 0.0f;
+	ScrollbarOffset = 0.0f;
+	ProfileData = InArgs._ProfileData.Get();
+	OnBarGraphSelectionChangedDelegate = InArgs._OnBarGraphSelectionChanged;
+	OnBarGraphExpansionChangedDelegate = InArgs._OnBarGraphExpansionChanged;
+	OnBarEventSelectionChangedDelegate = InArgs._OnBarEventSelectionChanged;
+	OnBarGraphContextMenuDelegate = InArgs._OnBarGraphContextMenu;
+	bSuppressBarGraphSelectionChangedDelegate = false;
+
+	ViewMode = EVisualizerViewMode::Hierarchical;
+	SelectedBarGraph = ProfileData;
+	CreateDataView();
+
+	// Drop down menu
+	const bool bInShouldCloseWindowAfterMenuSelection = true;
+	FMenuBuilder ViewMenuBuilder( bInShouldCloseWindowAfterMenuSelection, NULL );
+	{
+		FUIAction Action( FExecuteAction::CreateSP( this, &SBarVisualizer::SetViewMode, EVisualizerViewMode::Hierarchical ),
+				FCanExecuteAction(),
+				FIsActionChecked::CreateSP( this, &SBarVisualizer::CheckViewMode, EVisualizerViewMode::Hierarchical ) );
+
+		ViewMenuBuilder.AddMenuEntry( NSLOCTEXT("SBarVisualizer", "Hierarchical", "Hierarchical"), FText(), FSlateIcon(), Action, NAME_None, EUserInterfaceActionType::Check );
+	}
+
+	{
+		FUIAction Action( FExecuteAction::CreateSP( this, &SBarVisualizer::SetViewMode, EVisualizerViewMode::Flat ),
+				FCanExecuteAction(),
+				FIsActionChecked::CreateSP( this, &SBarVisualizer::CheckViewMode, EVisualizerViewMode::Flat ) );
+
+		ViewMenuBuilder.AddMenuEntry( NSLOCTEXT("SBarVisualizer", "Flat", "Flat"), FText(), FSlateIcon(), Action, NAME_None, EUserInterfaceActionType::Check );
+	}
+
+	const FSlateBrush* HomeButtonBrush = FProfileVisualizerStyle::Get()->GetBrush( "ProfileVisualizer.Home" );
+	const FSlateBrush* ToParentButtonBrush = FProfileVisualizerStyle::Get()->GetBrush( "ProfileVisualizer.ToParent" );
+
+	this->ChildSlot
+	[
+		SNew( SVerticalBox )
+		+SVerticalBox::Slot().AutoHeight() .Padding( 2 ) .VAlign( VAlign_Fill )
+		[
+			SNew( SBorder )
+			.BorderImage( FProfileVisualizerStyle::Get()->GetBrush("StatsHeader") )
+			.ForegroundColor( FProfileVisualizerStyle::Get()->GetSlateColor("DefaultForeground") )
+			[
+				SNew(SHorizontalBox)
+				+SHorizontalBox::Slot().AutoWidth() .Padding( 2 ) .HAlign( HAlign_Left )
+				[
+					SNew(SButton)
+					.ButtonStyle( FCoreStyle::Get(), "NoBorder" )
+					.ForegroundColor( FSlateColor::UseForeground() )
+					.ContentPadding(FMargin(0))
+					.Visibility( this, &SBarVisualizer::GetToParentButtonVisibility )
+					.OnClicked( this, &SBarVisualizer::OnToParentClicked )
+					[
+						SNew(SImage)
+						.Image( ToParentButtonBrush )
+					]
+				]
+				+SHorizontalBox::Slot().AutoWidth().Padding( 2 ).HAlign( HAlign_Left )
+				[
+					SNew(SButton)
+					.ButtonStyle( FCoreStyle::Get(), "NoBorder" )
+					.HAlign(HAlign_Center)
+					.VAlign(VAlign_Center)
+					.ForegroundColor( FSlateColor::UseForeground() )
+					.ContentPadding(FMargin(0))
+					.Visibility( this, &SBarVisualizer::GetHomeButtonVisibility )
+					.OnClicked( this, &SBarVisualizer::OnHomeClicked )
+					[
+						SNew(SImage)
+						.Image( HomeButtonBrush )
+					]
+				]
+				+SHorizontalBox::Slot() .Padding( 2 ) .FillWidth( 20 ) .HAlign( HAlign_Fill )
+				[
+					SNew( STextBlock )
+					.Text( this, &SBarVisualizer::GetSelectedCategoryName )					
+				]
+				+SHorizontalBox::Slot().AutoWidth() .HAlign( HAlign_Right ) .Padding( FMargin(1.0f, 2.0f, 1.0f, 2.0f) )
+				[
+					SNew( SComboButton )
+					//.ToolTipText(NSLOCTEXT("PropertyEditor", "ResetToDefaultToolTip", "Reset to Default").ToString())
+					.HasDownArrow( false )
+					.ButtonStyle( FCoreStyle::Get(), "NoBorder" )
+					.ContentPadding(0)
+					.ButtonContent()
+					[
+						SNew(SImage)
+						.Image( FProfileVisualizerStyle::Get()->GetBrush("ProfileVisualizer.MenuDropdown") )
+					]
+					.MenuContent()
+					[
+						ViewMenuBuilder.MakeWidget()
+					]
+				]	
+			]
+		]
+		+SVerticalBox::Slot() .Padding( 2 ) .FillHeight( 1 ) .VAlign( VAlign_Fill )
+		[
+			SNew(SHorizontalBox)
+			+SHorizontalBox::Slot() .Padding( 2 ) .FillWidth( 1 ) .HAlign( HAlign_Fill )
+			[
+				// List of thread graphs
+				SAssignNew( BarGraphsList, SListView< TSharedPtr< FVisualizerEvent > > )
+				// List view items are this tall
+				.ItemHeight( 24 )
+				// Tell the list view where to get its source data
+				.ListItemsSource( &ProfileDataView )
+				// When the list view needs to generate a widget for some data item, use this method
+				.OnGenerateRow( this, &SBarVisualizer::OnGenerateWidgetForList )
+				// Single selection mode
+				.SelectionMode( ESelectionMode::Single )
+				// Selection changed callback
+				.OnSelectionChanged( this, &SBarVisualizer::OnBarGraphSelectionChanged )
+			]
+		]
+		+SVerticalBox::Slot().AutoHeight() .Padding( 2 ) .VAlign( VAlign_Fill )
+		[
+			SAssignNew( Timeline, STimeline )
+			.MinValue(0.0f)
+			.MaxValue(SelectedBarGraph->DurationMs)
+		]
+		+SVerticalBox::Slot().AutoHeight() .Padding( 2 ) .VAlign( VAlign_Fill )
+		[
+			SAssignNew( ScrollBar, SScrollBar )
+			.Orientation( Orient_Horizontal )
+			.OnUserScrolled( this, &SBarVisualizer::ScrollBar_OnUserScrolled )
+		]
+		+SVerticalBox::Slot().AutoHeight() .Padding( 2 ) .VAlign( VAlign_Fill )
+		[
+			SNew( SHorizontalBox )
+			+SHorizontalBox::Slot()
+			.Padding( 2 )
+			.FillWidth( 1 )
+			.HAlign( HAlign_Fill )
+			[
+				SAssignNew( ZoomLabel, STextBlock )
+					.Text( this, &SBarVisualizer::GetZoomLabel )
+			]
+			+SHorizontalBox::Slot()
+			.Padding( 2 )
+			.FillWidth( 5 )
+			.HAlign( HAlign_Fill )
+			[
+				SNew( SSlider )
+				.Value( this, &SBarVisualizer::GetZoomValue )
+				.OnValueChanged( this, &SBarVisualizer::OnSetZoomValue )
+			]
+		]
+	];
+
+
+	ScrollBar->SetState(0.0f, 1.0f);
+}
+END_SLATE_FUNCTION_BUILD_OPTIMIZATION
+
+void SBarVisualizer::ScrollBar_OnUserScrolled( float InScrollOffsetFraction )
+{
+	if( ZoomSliderValue > 0.0f )
+	{
+		const float MaxOffset = GetMaxScrollOffsetFraction();
+		const float MaxGraphOffset = GetMaxGraphOffset();
+		InScrollOffsetFraction = FMath::Clamp( InScrollOffsetFraction, 0.0f, MaxOffset );
+		float GraphOffset = -( InScrollOffsetFraction / MaxOffset ) * MaxGraphOffset;
+
+		ScrollBar->SetState( InScrollOffsetFraction, 1.0f / GetZoom() );
+
+		for( int32 Index = 0; Index < Graphs.Num(); Index++ )
+		{
+			Graphs[ Index ]->SetOffset( GraphOffset );
+		}
+
+		Timeline->SetOffset( GraphOffset );
+
+		ScrollbarOffset = GraphOffset;
+	}
+}
+
+FText SBarVisualizer::GetZoomLabel() const
+{
+	static const FNumberFormattingOptions ZoomFormatOptions = FNumberFormattingOptions()
+		.SetMinimumFractionalDigits(2)
+		.SetMaximumFractionalDigits(2);
+	return FText::Format( NSLOCTEXT("ProfileVisualizer", "ZoomLabelFmt", "Zoom: {0}x"), FText::AsNumber(GetZoom(), &ZoomFormatOptions) );
+}
+
+float SBarVisualizer::GetZoomValue() const
+{
+	return ZoomSliderValue;
+}
+
+void SBarVisualizer::OnSetZoomValue( float NewValue )
+{
+	const float PrevZoom = GetZoom();
+	const float PrevVisibleRange = 1.0f / PrevZoom;
+		
+	ZoomSliderValue = NewValue;
+	const float Zoom = GetZoom();
+
+	float GraphOffset = 0.0f;
+	float ScrollOffsetFraction = 0.0f;
+
+	if( Graphs.Num() )
+	{
+		const float MaxOffset = GetMaxScrollOffsetFraction();
+		const float MaxGraphOffset = GetMaxGraphOffset();
+		const float PrevGraphOffset = -Graphs[ 0 ]->GetOffset();
+		GraphOffset = FMath::Clamp( -Graphs[ 0 ]->GetOffset(), 0.0f, MaxGraphOffset );
+			
+		const float VisibleRange = 1.0f / GetZoom();
+		const float PrevGraphCenterValue = PrevGraphOffset / PrevZoom + PrevVisibleRange * 0.5f;
+		const float NewGraphCenterValue = GraphOffset / Zoom + VisibleRange * 0.5f;
+		GraphOffset += ( PrevGraphCenterValue - NewGraphCenterValue ) * Zoom;
+		GraphOffset = FMath::Clamp( GraphOffset, 0.0f, MaxGraphOffset );
+
+		ScrollOffsetFraction = FMath::Clamp( MaxOffset * GraphOffset / MaxGraphOffset, 0.0f, MaxOffset );
+	}
+
+	ScrollBar->SetState( ScrollOffsetFraction, 1.0f / Zoom );
+
+	for( int32 Index = 0; Index < Graphs.Num(); Index++ )
+	{
+		Graphs[ Index ]->SetZoom( Zoom );
+		Graphs[ Index ]->SetOffset( -GraphOffset );
+	}
+
+	Timeline->SetZoom( Zoom );
+	Timeline->SetOffset( -GraphOffset );
+
+	ScrollbarOffset = -GraphOffset;
+}
+
+void SBarVisualizer::OnBarGraphSelectionChanged( TSharedPtr< FVisualizerEvent > Selection, ESelectInfo::Type SelectInfo )
+{
+	if( Selection.IsValid() )
+	{
+		BarGraphsList->RequestListRefresh();
+	}
+
+	if( !bSuppressBarGraphSelectionChangedDelegate )
+	{
+		OnBarGraphSelectionChangedDelegate.ExecuteIfBound( Selection );
+	}
+}
+
+void SBarVisualizer::ClearBarSelection( TSharedPtr< FVisualizerEvent > GraphEvents, TSharedPtr<FVisualizerEvent> Selection )
+{
+	for( int32 ChildIndex = 0; ChildIndex < GraphEvents->Children.Num(); ChildIndex++ )
+	{
+		// Don't clear selection on the event to be selected
+		if ( Selection != GraphEvents->Children[ ChildIndex ] )
+		{
+			GraphEvents->Children[ ChildIndex ]->IsSelected = false;
+		}
+
+		ClearBarSelection( GraphEvents->Children[ ChildIndex ], Selection );
+	}
+}
+
+void SBarVisualizer::OnBarEventSelectionChanged( TSharedPtr<FVisualizerEvent> Selection, ESelectInfo::Type SelectInfo, int32 BarId )
+{
+	HandleEventSelectionChanged( Selection );
+	OnBarEventSelectionChangedDelegate.ExecuteIfBound( BarId, Selection );
+}
+
+TSharedPtr< FVisualizerEvent > SBarVisualizer::FindSelectedEventsParent( TArray< TSharedPtr< FVisualizerEvent > >& BarGraphs, TSharedPtr< FVisualizerEvent > Selection )
+{
+	for( int32 Index = 0; Index < BarGraphs.Num(); Index++ )
+	{
+		if( BarGraphs[ Index ]->Children.Contains( Selection ) )
+		{
+			return BarGraphs[ Index ];
+		}
+		
+		TSharedPtr< FVisualizerEvent > SelectionParent = FindSelectedEventsParent( BarGraphs[ Index ]->Children, Selection );
+		if ( SelectionParent.IsValid() )
+		{
+			return SelectionParent;
+		}
+	}
+
+	return TSharedPtr< FVisualizerEvent >();
+}
+
+void SBarVisualizer::HandleEventSelectionChanged( TSharedPtr< FVisualizerEvent > Selection )
+{
+	// Clear any selected events from the other bar graphs
+	const int32 BarId = Selection.IsValid() ? Selection->Category : INDEX_NONE;
+	int32 BarIndex = INDEX_NONE;
+
+	TSharedPtr< FVisualizerEvent > Root = ProfileData;
+	while ( Root->ParentEvent.IsValid() )
+	{
+		Root = Root->ParentEvent;
+	}
+
+	ClearBarSelection( Root, Selection );
+
+	// Since we're changing the selection as a result of selection change we don't want to create an infinite loop
+	bSuppressBarGraphSelectionChangedDelegate = true;
+
+	BarGraphsList->ClearSelection();
+
+	if( Selection.IsValid() )
+	{
+		// Check if one of the bars has been selected
+		bool BarGraphFound = false;
+		for( int32 Index = 0; Index < ProfileDataView.Num() && BarGraphFound == false; Index++ )
+		{
+			TSharedPtr< FVisualizerEvent > BarGraph = ProfileDataView[ Index ];
+			if( BarGraph->EventName == Selection->EventName )
+			{
+				BarGraphsList->SetSelection( BarGraph );
+				BarGraphsList->RequestScrollIntoView( BarGraph );
+				BarGraphFound = true;
+			}
+		}
+	}
+
+	bSuppressBarGraphSelectionChangedDelegate = false;
+}
+
+BEGIN_SLATE_FUNCTION_BUILD_OPTIMIZATION
+TSharedRef<ITableRow> SBarVisualizer::OnGenerateWidgetForList( TSharedPtr<FVisualizerEvent> InItem, const TSharedRef<STableViewBase>& OwnerTable )
+{
+	TSharedPtr<SGraphBar> RowGraph = SNew( SGraphBar )
+				.OnSelectionChanged( this, &SBarVisualizer::OnBarEventSelectionChanged, InItem->Category )
+				.OnGeometryChanged( this, &SBarVisualizer::OnBarGeometryChanged );
+
+	const double EventsStartTime = InItem->ParentEvent.IsValid() ? InItem->ParentEvent->Start : 0.0;
+	const double EventsEndTime = InItem->ParentEvent.IsValid() ? InItem->ParentEvent->Duration : 1.0;
+	RowGraph->SetEvents( InItem->Children, EventsStartTime, EventsEndTime );
+	Graphs.Add( RowGraph );
+	RowGraph->SetZoom( GetZoom() );
+	RowGraph->SetOffset( ScrollbarOffset );
+
+	TSharedPtr<SWidget> BarTitle;
+
+	if( IsExpandable( InItem ) )
+	{
+		BarTitle = SNew( SHorizontalBox )
+			+SHorizontalBox::Slot()
+			.AutoWidth()
+			[
+				SNew(SButton)
+				.ButtonStyle( FCoreStyle::Get(), "NoBorder" )
+				.ForegroundColor( FSlateColor::UseForeground() )
+				.ContentPadding(FMargin(0))
+				.OnClicked( this, &SBarVisualizer::ExpandBar, InItem )
+				[
+					SNew(SBorder) 
+					.BorderImage( FCoreStyle::Get().GetBrush("NoBorder") )
+					.HAlign(HAlign_Center)
+					.VAlign(VAlign_Center)
+					.Padding(0)						
+					[
+						SNew(SImage)
+						.Image( FCoreStyle::Get().GetBrush( "TreeArrow_Collapsed" ) )
+					]
+				]
+			]
+			+SHorizontalBox::Slot()
+			.AutoWidth()
+			[
+				SNew( STextBlock ).Text( FText::FromString(InItem->EventName) )
+			];
+	}
+	else
+	{
+		BarTitle = SNew( STextBlock ).Text( FText::FromString(InItem->EventName) );
+	}
+
+	return SNew( STableRow< TSharedPtr< FVisualizerEvent > >, OwnerTable )
+		[
+			SNew( SBorder )
+			.Padding(0)
+			.BorderImage( FCoreStyle::Get().GetBrush("NoBorder") )
+			// Handle right-click event for context menu
+			.OnMouseButtonDown( this, &SBarVisualizer::OnBarRightClicked, InItem )
+			[
+				SNew( SVerticalBox )
+				+SVerticalBox::Slot()
+				.AutoHeight()
+				[
+					BarTitle.ToSharedRef()
+				]
+				+SVerticalBox::Slot()
+				.AutoHeight()
+				[
+					RowGraph.ToSharedRef()
+				]
+			]
+		];
+}
+END_SLATE_FUNCTION_BUILD_OPTIMIZATION
+
+/** Checks if the selected event has children with children */
+bool SBarVisualizer::IsExpandable( TSharedPtr< FVisualizerEvent > InEvent )
+{
+	for( int32 ChildIndex = 0; ChildIndex < InEvent->Children.Num(); ChildIndex++ )
+	{
+		if( InEvent->Children[ ChildIndex ]->Children.Num() > 0 )
+		{
+			return true;
+		}
+	}
+	return false;
+}
+
+void SBarVisualizer::OnGetChildrenForList( TSharedPtr<FVisualizerEvent> InItem, TArray<TSharedPtr<FVisualizerEvent> >& OutChildren )
+{
+	OutChildren = InItem->Children;
+}
+
+FReply SBarVisualizer::OnBarRightClicked( const FGeometry& BarGeometry, const FPointerEvent& MouseEvent, TSharedPtr<FVisualizerEvent> Selection )
+{
+	if( MouseEvent.GetEffectingButton() == EKeys::RightMouseButton  )
+	{
+		if(OnBarGraphContextMenuDelegate.IsBound())
+		{
+			// Forward the event to the visualizer main frame
+			// Disabled for now, may be useful in the future
+			return FReply::Handled();
+		}
+	}
+
+	return FReply::Unhandled();
+}
+
+void SBarVisualizer::CreateFlattenedData( TSharedPtr< FVisualizerEvent > InData, TArray< TSharedPtr< FVisualizerEvent > >& FlattenedData )
+{
+	for( int32 EventIndex = 0; EventIndex < InData->Children.Num(); EventIndex++ )
+	{
+		TSharedPtr< FVisualizerEvent > Event = InData->Children[ EventIndex ];
+		TSharedPtr< FVisualizerEvent > FlattendedDataSet;
+
+		if( Event->Category >= FlattenedData.Num() || FlattenedData[ Event->Category ].IsValid() == false )
+		{
+			// Fill with empty categories if necessary
+			for( int32 CategoryIndex = FlattenedData.Num(); CategoryIndex <= Event->Category; CategoryIndex++ )
+			{
+				FlattenedData.Add( TSharedPtr< FVisualizerEvent >() );
+			}
+
+			// Create a separate data set for this category (thread)
+			FlattendedDataSet = MakeShareable< FVisualizerEvent >( new FVisualizerEvent( 0.0, 1.0, 0.0, Event->Category, Event->EventName ) );
+			FlattendedDataSet->Category = Event->Category;
+			FlattenedData[ Event->Category ] = FlattendedDataSet;
+
+			// Get the category name by looking for the first occurence of the category in the tree
+			for( TSharedPtr< FVisualizerEvent > CategoryData = Event->ParentEvent; CategoryData.IsValid() && FlattendedDataSet->EventName.IsEmpty(); CategoryData = CategoryData->ParentEvent )
+			{
+				if( CategoryData->ParentEvent.IsValid() == false || CategoryData->ParentEvent->Category != Event->Category )
+				{
+					FlattendedDataSet->EventName = CategoryData->EventName;
+				}
+			}
+		}
+		else
+		{
+			FlattendedDataSet = FlattenedData[ Event->Category ];
+		}
+
+		// Fill with events
+		if( InData->Children[ EventIndex ]->Children.Num() == 0 )
+		{
+			FlattendedDataSet->Children.Add( InData->Children[ EventIndex ] );
+		}
+	}
+
+	for( int32 ChildIndex = 0; ChildIndex < InData->Children.Num(); ChildIndex++ )
+	{
+		CreateFlattenedData( InData->Children[ ChildIndex ], FlattenedData );
+	}
+}
+
+void SBarVisualizer::CreateDataView()
+{
+	// Each time new data set is being displayed, clear all cached Bar Graphs
+	Graphs.Empty();
+
+	if( ViewMode == EVisualizerViewMode::Flat )
+	{
+		ProfileDataView.Empty();		
+		TArray< TSharedPtr< FVisualizerEvent > > FlattenedBarGraphData;
+
+		// At the top level each bar may represent a different subset of data (like a thread) so only flatten within one data set
+		CreateFlattenedData( ProfileData, FlattenedBarGraphData );
+
+		ProfileDataView.Append( FlattenedBarGraphData );
+	}
+	else
+	{
+		// Default to hierarchical. Choose the selected subset
+		if( SelectedBarGraph.IsValid() == false )
+		{
+			ProfileDataView.Empty( 1 );
+			ProfileDataView.Add( ProfileData );
+		}
+		else
+		{
+			ProfileDataView.Empty( SelectedBarGraph->Children.Num() );
+
+			// Get all leaf events into one bar graph
+			TArray< TSharedPtr< FVisualizerEvent > > LeafEvents; 
+			for( int32 EventIndex = 0; EventIndex < SelectedBarGraph->Children.Num(); EventIndex++ )
+			{
+				TSharedPtr< FVisualizerEvent > Event = SelectedBarGraph->Children[ EventIndex ];
+				if( Event->Children.Num() == 0 )
+				{
+					LeafEvents.Add( Event );
+				}
+				else
+				{
+					ProfileDataView.Add( Event );
+				}
+			}
+
+			if( LeafEvents.Num() > 0 )
+			{
+				TSharedPtr< FVisualizerEvent > LeafEventsBarGraph( new FVisualizerEvent( SelectedBarGraph->Start, SelectedBarGraph->Duration, SelectedBarGraph->DurationMs, SelectedBarGraph->Category, SelectedBarGraph->EventName + TEXT(" Leaf Events") ) );
+				LeafEventsBarGraph->ParentEvent = SelectedBarGraph;
+				LeafEventsBarGraph->Children = LeafEvents;
+
+				ProfileDataView.Add( LeafEventsBarGraph );
+			}
+		}
+	}
+}
+
+void SBarVisualizer::SetViewMode( EVisualizerViewMode::Type InMode )
+{	
+	ViewMode = InMode;
+
+	SelectedBarGraph = ProfileData;
+	CreateDataView();
+
+	BarGraphsList->RequestListRefresh();
+
+	OnBarGraphExpansionChangedDelegate.ExecuteIfBound( SelectedBarGraph );
+}
+
+FReply SBarVisualizer::OnToParentClicked()
+{
+	if( SelectedBarGraph.IsValid() && SelectedBarGraph->ParentEvent.IsValid() )
+	{
+		SelectedBarGraph = SelectedBarGraph->ParentEvent;
+		CreateDataView();
+
+		BarGraphsList->RequestListRefresh();
+
+		OnBarGraphExpansionChangedDelegate.ExecuteIfBound( SelectedBarGraph );
+
+		AdjustTimeline( SelectedBarGraph );
+	}
+	return FReply::Handled();
+}
+
+FReply SBarVisualizer::OnHomeClicked()
+{
+	SelectedBarGraph = ProfileData;
+	CreateDataView();
+	BarGraphsList->RequestListRefresh();
+
+	OnBarGraphExpansionChangedDelegate.ExecuteIfBound( SelectedBarGraph );
+	AdjustTimeline( SelectedBarGraph );
+	return FReply::Handled();
+}
+
+void SBarVisualizer::OnBarGeometryChanged( FGeometry Geometry )
+{
+	Timeline->SetDrawingGeometry( Geometry );
+}
+
+FText SBarVisualizer::GetSelectedCategoryName() const
+{
+	if( SelectedBarGraph.IsValid() )
+	{
+		FString EventName( SelectedBarGraph->EventName );
+
+		for( TSharedPtr< FVisualizerEvent > BarGraph = SelectedBarGraph->ParentEvent; BarGraph.IsValid(); BarGraph = BarGraph->ParentEvent )
+		{
+			EventName = BarGraph->EventName + TEXT("\\") + EventName;
+		}
+		return FText::FromString(EventName);
+	}
+	else
+	{
+		return NSLOCTEXT("SBarVisualizer", "Frame", "Frame");
+	}
+}
+
+EVisibility SBarVisualizer::GetHomeButtonVisibility() const
+{
+	if( SelectedBarGraph.IsValid() && SelectedBarGraph->ParentEvent.IsValid() && SelectedBarGraph->ParentEvent->ParentEvent.IsValid() )
+	{
+		return EVisibility::Visible;
+	}
+	else
+	{
+		return EVisibility::Collapsed;
+	}
+}
+
+EVisibility SBarVisualizer::GetToParentButtonVisibility() const
+{
+	if( SelectedBarGraph.IsValid() && SelectedBarGraph->ParentEvent.IsValid() )
+	{
+		return EVisibility::Visible;
+	}
+	else
+	{
+		return EVisibility::Collapsed;
+	}
+}
+
+FReply SBarVisualizer::ExpandBar( TSharedPtr<FVisualizerEvent> BarGraphEvents )
+{
+	if( BarGraphEvents.IsValid() && BarGraphEvents->Children.Num() > 0 )
+	{
+		SelectedBarGraph = BarGraphEvents;
+		CreateDataView();
+		BarGraphsList->RequestListRefresh();
+
+		OnBarGraphExpansionChangedDelegate.ExecuteIfBound( SelectedBarGraph );
+
+		// Set Timeline scale appropriate for the selected events
+		AdjustTimeline( BarGraphEvents );
+	}
+	return FReply::Handled();
+}
+
+void SBarVisualizer::AdjustTimeline( TSharedPtr< FVisualizerEvent > InEvent )
+{
+	check( InEvent.IsValid() );
+
+	const double TotalTimeMs = InEvent->DurationMs / InEvent->Duration;
+	const double StartMs = InEvent->Start * TotalTimeMs;
+	Timeline->SetMinMaxValues( StartMs, StartMs + InEvent->DurationMs );
+}
diff --git a/Engine/Source/Developer/ProfileVisualizer/Private/SEventsTree.cpp b/Engine/Source/Developer/ProfileVisualizer/Private/SEventsTree.cpp
new file mode 100644
--- /dev/null
+++ b/Engine/Source/Developer/ProfileVisualizer/Private/SEventsTree.cpp
@@ -0,0 +1,671 @@
+// Copyright Epic Games, Inc. All Rights Reserved.
+
+#include "SEventsTree.h"
+#include "SlateOptMacros.h"
+#include "Textures/SlateIcon.h"
+#include "Framework/Commands/UIAction.h"
+#include "Framework/MultiBox/MultiBoxBuilder.h"
+#include "Widgets/Input/SSearchBox.h"
+
+#define LOCTEXT_NAMESPACE "SEventsTree"
+
+FName SEventsTree::NAME_NameColumn = FName(*NSLOCTEXT("ProfileVisualizer", "ColumnName", "Name").ToString());
+FName SEventsTree::NAME_DurationColumn = FName( *NSLOCTEXT("ProfileVisualizer", "ColumnDuration", "Duration").ToString() );
+
+void SEventItem::Construct( const FArguments& InArgs, const TSharedRef<STableViewBase>& InOwnerTableView )
+{
+	EventName = InArgs._EventName;
+	EventDuration = InArgs._EventDuration;
+
+	SMultiColumnTableRow< TSharedPtr< FVisualizerEvent > >::Construct( FSuperRowType::FArguments(), InOwnerTableView );	
+}
+
+BEGIN_SLATE_FUNCTION_BUILD_OPTIMIZATION
+TSharedRef< SWidget > SEventItem::GenerateWidgetForColumn( const FName& ColumnName )
+{
+	if( ColumnName == SEventsTree::NAME_NameColumn )
+	{
+		return  
+		SNew(SHorizontalBox)
+		+SHorizontalBox::Slot()
+		.AutoWidth()
+		[
+			SNew( SExpanderArrow, SharedThis(this) )
+		]
+		+SHorizontalBox::Slot()
+		.AutoWidth()
+		[
+			SNew( STextBlock )
+			.Text( FText::FromString(EventName) )
+		];
+	}
+	else if( ColumnName == SEventsTree::NAME_DurationColumn )
+	{
+		return SNew( STextBlock )
+			.Text( this, &SEventItem::GetDurationText );
+	}
+	else
+	{
+		return
+		SNew( STextBlock )
+		. Text(FText::Format( LOCTEXT("UnsupportedColumnFmt", "Unsupported Column: {0}"), FText::FromName(ColumnName) ) );
+	}
+}
+END_SLATE_FUNCTION_BUILD_OPTIMIZATION
+
+/////////////////////////////////////////////////////////////////////////////////////////////////////////////
+
+BEGIN_SLATE_FUNCTION_BUILD_OPTIMIZATION
+void SEventsTree::Construct( const FArguments& InArgs )
+{
+	DurationUnits = EVisualizerTimeUnits::Milliseconds;
+	ViewMode = EVisualizerViewMode::Hierarchical;
+	bSuppressSelectionChangedEvent = false;
+
+	OnEventSelectionChangedDelegate = InArgs._OnEventSelectionChanged;
+	ProfileData = InArgs._ProfileData.Get();
+
+	// Duration column drop down menu
+	const bool bInShouldCloseWindowAfterMenuSelection = true;
+	FMenuBuilder MenuBuilder( bInShouldCloseWindowAfterMenuSelection, NULL );
+	{
+		FUIAction Action( FExecuteAction::CreateSP( this, &SEventsTree::SetDurationUnits, EVisualizerTimeUnits::Microseconds ),
+				FCanExecuteAction(),
+				FIsActionChecked::CreateSP( this, &SEventsTree::CheckDurationUnits, EVisualizerTimeUnits::Microseconds ) );
+
+		MenuBuilder.AddMenuEntry( LOCTEXT("Microseconds", "Microseconds"), FText(), FSlateIcon(), Action, NAME_None, EUserInterfaceActionType::Check );
+	}
+
+	{
+		FUIAction Action( FExecuteAction::CreateSP( this, &SEventsTree::SetDurationUnits, EVisualizerTimeUnits::Milliseconds ),
+				FCanExecuteAction(),
+				FIsActionChecked::CreateSP( this, &SEventsTree::CheckDurationUnits, EVisualizerTimeUnits::Milliseconds ) );
+
+		MenuBuilder.AddMenuEntry( LOCTEXT("Milliseconds", "Milliseconds"), FText(), FSlateIcon(), Action, NAME_None, EUserInterfaceActionType::Check );
+	}
+
+	{
+		FUIAction Action( FExecuteAction::CreateSP( this, &SEventsTree::SetDurationUnits, EVisualizerTimeUnits::Seconds ),
+				FCanExecuteAction(),
+				FIsActionChecked::CreateSP( this, &SEventsTree::CheckDurationUnits, EVisualizerTimeUnits::Seconds ) );
+
+		MenuBuilder.AddMenuEntry( LOCTEXT("Seconds", "Seconds"), FText(), FSlateIcon(), Action, NAME_None, EUserInterfaceActionType::Check );
+	}
+
+	// Name column drop down menu
+	FMenuBuilder NameMenuBuilder( bInShouldCloseWindowAfterMenuSelection, NULL );
+	{
+		FUIAction Action( FExecuteAction::CreateSP( this, &SEventsTree::SetViewMode, EVisualizerViewMode::Hierarchical ),
+				FCanExecuteAction(),
+				FIsActionChecked::CreateSP( this, &SEventsTree::CheckViewMode, EVisualizerViewMode::Hierarchical ) );
+
+		NameMenuBuilder.AddMenuEntry( LOCTEXT("Hierarchical", "Hierarchical"), FText(), FSlateIcon(), Action, NAME_None, EUserInterfaceActionType::Check );
+	}
+
+	{
+		FUIAction Action( FExecuteAction::CreateSP( this, &SEventsTree::SetViewMode, EVisualizerViewMode::Flat ),
+				FCanExecuteAction(),
+				FIsActionChecked::CreateSP( this, &SEventsTree::CheckViewMode, EVisualizerViewMode::Flat ) );
+
+		NameMenuBuilder.AddMenuEntry( LOCTEXT("Flat", "Flat"), FText(), FSlateIcon(), Action, NAME_None, EUserInterfaceActionType::Check );
+	}
+
+	{
+		FUIAction Action( FExecuteAction::CreateSP( this, &SEventsTree::SetViewMode, EVisualizerViewMode::Coalesced ),
+				FCanExecuteAction(),
+				FIsActionChecked::CreateSP( this, &SEventsTree::CheckViewMode, EVisualizerViewMode::Coalesced ) );
+
+		NameMenuBuilder.AddMenuEntry( LOCTEXT("Coalesced", "Coalesced"), FText(), FSlateIcon(), Action, NAME_None, EUserInterfaceActionType::Check );
+	}
+
+	{
+		FUIAction Action( FExecuteAction::CreateSP( this, &SEventsTree::SetViewMode, EVisualizerViewMode::FlatCoalesced ),
+				FCanExecuteAction(),
+				FIsActionChecked::CreateSP( this, &SEventsTree::CheckViewMode, EVisualizerViewMode::FlatCoalesced ) );
+
+		NameMenuBuilder.AddMenuEntry( LOCTEXT("FlatCoalesced", "Flat Coalesced"), FText(), FSlateIcon(), Action, NAME_None, EUserInterfaceActionType::Check );
+	}
+
+	this->ChildSlot
+	[
+		SNew( SVerticalBox )
+		+SVerticalBox::Slot().AutoHeight().Padding( 1.0f, 0.0f, 1.0f, 2.0f )
+		[
+			SNew( SSearchBox )
+			.ToolTipText( NSLOCTEXT("ProfileVisualizer", "FilterSearchHint", "Type here to search events.") )
+			.OnTextChanged( this, &SEventsTree::OnFilterTextChanged )
+			.OnTextCommitted( this, &SEventsTree::OnFilterTextCommitted )
+		]
+		+SVerticalBox::Slot().Padding( 2 ).FillHeight( 1 ).VAlign( VAlign_Fill )
+		[
+			SNew( SHorizontalBox )
+			+SHorizontalBox::Slot().Padding( 2 ).FillWidth( 1 ).HAlign( HAlign_Fill )
+			[
+				// List of all events for the selected thread
+				SAssignNew( EventsListView, STreeView< TSharedPtr< FVisualizerEvent > > )
+				// List view items are this tall
+				.ItemHeight( 12 )
+				// Tell the list view where to get its source data
+				.TreeItemsSource( &SelectedEventsView )
+				// When the list view needs to generate a widget for some data item, use this method
+				.OnGenerateRow( this, &SEventsTree::OnGenerateWidgetForEventsList )
+				// Given some DataItem, this is how we find out if it has any children and what they are.
+				.OnGetChildren( this, &SEventsTree::OnGetChildrenForEventsList )
+				// Selection mode
+				.SelectionMode( ESelectionMode::Single )
+				// Selection callback
+				.OnSelectionChanged( this, &SEventsTree::OnEventSelectionChanged )
+				.HeaderRow
+				(
+					SNew( SHeaderRow )
+					+ SHeaderRow::Column( NAME_NameColumn )
+					.DefaultLabel( NSLOCTEXT("ProfileVisualizer", "ColumnName", "Name") ) 
+					.SortMode( TAttribute< EColumnSortMode::Type >::Create( TAttribute< EColumnSortMode::Type >::FGetter::CreateSP( this, &SEventsTree::GetColumnSortMode, NAME_NameColumn ) ) )
+					.OnSort( FOnSortModeChanged::CreateSP( this, &SEventsTree::OnColumnSortModeChanged ) )
+					.FillWidth( 1.0f )
+					.MenuContent()
+					[
+						NameMenuBuilder.MakeWidget()
+					]
+
+					+ SHeaderRow::Column( NAME_DurationColumn )
+					.DefaultLabel( this, &SEventsTree::GetDurationColumnTitle ) 
+					.SortMode( TAttribute< EColumnSortMode::Type >::Create( TAttribute< EColumnSortMode::Type >::FGetter::CreateSP( this, &SEventsTree::GetColumnSortMode, NAME_DurationColumn ) ) )
+					.OnSort( FOnSortModeChanged::CreateSP( this, &SEventsTree::OnColumnSortModeChanged ) )
+					.FixedWidth( 128.0f )
+					.MenuContent()
+					[
+						MenuBuilder.MakeWidget()
+					]
+				)
+			]
+		]
+	];
+}
+END_SLATE_FUNCTION_BUILD_OPTIMIZATION
+
+EColumnSortMode::Type SEventsTree::GetColumnSortMode( const FName ColumnId )
+{
+	if ( SortByColumn != ColumnId )
+	{
+		return EColumnSortMode::None;
+	}
+
+	return SortMode;
+}
+
+FText SEventsTree::GetDurationColumnTitle() const
+{
+	static const FText Units[] = { NSLOCTEXT("ProfileVisualizer", "microseconds", "microseconds"), NSLOCTEXT("ProfileVisualizer", "milliseconds", "ms"), NSLOCTEXT("ProfileVisualizer", "seconds", "s") };
+
+	return FText::Format( NSLOCTEXT("ProfileVisualizer", "ColumnDurationValue", "Duration ({0})"), Units[ DurationUnits ] );
+}
+
+void ClearEventsSelection( TArray< TSharedPtr< FVisualizerEvent > >& Events )
+{
+	for( int32 Index = 0; Index < Events.Num(); ++Index )
+	{
+		TSharedPtr< FVisualizerEvent > CurrentEvent = Events[ Index ];
+		CurrentEvent->IsSelected = false;
+
+		// Clear recursively
+		if( CurrentEvent->Children.Num() )
+		{
+			ClearEventsSelection( CurrentEvent->Children );
+		}
+	}
+}
+
+void SEventsTree::OnEventSelectionChanged( TSharedPtr< FVisualizerEvent > Selection, ESelectInfo::Type /*SelectInfo*/ )
+{
+	if( Selection.IsValid() )
+	{
+		ClearEventsSelection( SelectedEventsView );
+
+		Selection->IsSelected = true;
+
+		// Mirror the selection in the source events list
+		TSharedPtr< FVisualizerEvent > MappedSelection = ViewToEventsMap.FindChecked( Selection );
+		MappedSelection->IsSelected = true;
+
+		if( !bSuppressSelectionChangedEvent )
+		{
+		OnEventSelectionChangedDelegate.ExecuteIfBound( MappedSelection );
+	}
+	}
+	else if( !bSuppressSelectionChangedEvent )
+	{
+		OnEventSelectionChangedDelegate.ExecuteIfBound( Selection );
+	}
+}
+
+TSharedRef< ITableRow > SEventsTree::OnGenerateWidgetForEventsList( TSharedPtr< FVisualizerEvent > InItem, const TSharedRef<STableViewBase>& OwnerTable )
+{
+	return SNew( SEventItem, OwnerTable )
+			.EventName( InItem->EventName )
+			.EventDuration( this, &SEventsTree::GetEventDuration, InItem->DurationMs );
+}
+
+double SEventsTree::GetEventDuration( double InDurationMs ) const
+{
+	switch( DurationUnits )
+	{
+	case EVisualizerTimeUnits::Microseconds:
+		return InDurationMs * 1000.0;
+	case EVisualizerTimeUnits::Milliseconds:
+		return InDurationMs;
+	case EVisualizerTimeUnits::Seconds:
+		return InDurationMs * 0.001;
+	}
+
+	return InDurationMs;
+}
+
+void SEventsTree::SetDurationUnits( EVisualizerTimeUnits::Type InUnits )
+{	
+	DurationUnits = InUnits;
+	EventsListView->RequestTreeRefresh();
+}
+
+void SEventsTree::SetViewMode( EVisualizerViewMode::Type InMode )
+{	
+	ViewMode = InMode;
+
+	CreateSelectedEventsView();
+	SortEventsList();
+
+	EventsListView->RequestTreeRefresh();
+}
+
+void SEventsTree::OnGetChildrenForEventsList( TSharedPtr<FVisualizerEvent> InItem, TArray<TSharedPtr<FVisualizerEvent> >& OutChildren )
+{
+	OutChildren = InItem->Children;
+}
+
+void SEventsTree::OnColumnSortModeChanged( const EColumnSortPriority::Type SortPriority, const FName& ColumnId, const EColumnSortMode::Type InSortMode )
+{
+	SortByColumn = ColumnId;
+	SortMode = InSortMode;
+
+	SortEventsList();
+}
+
+void SEventsTree::SortEventsList( TArray< TSharedPtr< FVisualizerEvent > >& Events )
+{
+	// Sort taking the current settings into account
+	if( SortByColumn == NAME_NameColumn )
+	{
+		if( SortMode == EColumnSortMode::Ascending )
+		{
+			struct FCompareEventsByName
+			{
+				FORCEINLINE bool operator()( const TSharedPtr<FVisualizerEvent> A, const TSharedPtr<FVisualizerEvent> B ) const { return A->EventName < B->EventName; }
+			};
+			Events.Sort( FCompareEventsByName() );
+		}
+		else if( SortMode == EColumnSortMode::Descending )
+		{
+			struct FCompareEventsByNameDescending
+			{
+				FORCEINLINE bool operator()( const TSharedPtr<FVisualizerEvent> A, const TSharedPtr<FVisualizerEvent> B ) const { return B->EventName < A->EventName; }
+			};
+			Events.Sort( FCompareEventsByNameDescending() );
+		}
+	}
+	else if( SortByColumn == NAME_DurationColumn )
+	{
+		if( SortMode == EColumnSortMode::Ascending )
+		{
+			struct FCompareEventsByDuration
+			{
+				FORCEINLINE bool operator()( const TSharedPtr<FVisualizerEvent> A, const TSharedPtr<FVisualizerEvent> B ) const { return A->DurationMs < B->DurationMs; }
+			};
+			Events.Sort( FCompareEventsByDuration() );
+		}
+		else if( SortMode == EColumnSortMode::Descending )
+		{
+			struct FCompareEventsByDurationDescending
+			{
+				FORCEINLINE bool operator()( const TSharedPtr<FVisualizerEvent> A, const TSharedPtr<FVisualizerEvent> B ) const { return B->DurationMs < A->DurationMs; }
+			};
+			Events.Sort( FCompareEventsByDurationDescending() );
+		}
+	}
+
+	// Sort recursively
+	for( int32 Index = 0; Index < Events.Num(); Index++ )
+	{
+		if( Events[ Index ]->Children.Num() )
+		{
+			SortEventsList( Events[ Index ]->Children );
+		}
+	}
+}
+
+void SEventsTree::SortEventsList( void )
+{
+	SortEventsList( SelectedEventsView );
+
+	EventsListView->RequestTreeRefresh();
+
+	RestoreEventSelection( SelectedEventsView );
+}
+
+bool SEventsTree::RestoreEventSelection( TArray< TSharedPtr< FVisualizerEvent > >& Events )
+{
+	// Search for the selected event
+	for( int32 Index = 0; Index < Events.Num(); ++Index )
+	{
+		if( Events[ Index ]->IsSelected )
+		{
+			// Select it in the tree view widget
+			EventsListView->ClearSelection();
+			EventsListView->RequestTreeRefresh();
+			EventsListView->SetSelection( Events[ Index ] );
+			EventsListView->RequestScrollIntoView( Events[ Index ] );
+
+			return true;
+		}
+
+		// Search recursively
+		if( RestoreEventSelection( Events[ Index ]->Children ) )
+		{
+			return true;
+		}
+	}
+
+	return false;
+}
+
+void SEventsTree::HandleBarGraphSelectionChanged( TSharedPtr< FVisualizerEvent > Selection )
+{
+	if ( Selection.IsValid() )
+	{
+		for( int32 EventIndex = 0; EventIndex < SelectedEvents.Num(); EventIndex++ )
+		{
+			TSharedPtr<FVisualizerEvent> Event = SelectedEvents[ EventIndex ];
+			if( Selection->EventName == Event->EventName )
+			{
+				HandleBarEventSelectionChanged( SelectedThread->Category, Event );
+				break;
+			}
+		}
+	}
+}
+
+void SEventsTree::HandleBarGraphExpansionChanged( TSharedPtr< FVisualizerEvent > Selection )
+{
+	if ( Selection.IsValid() )
+	{
+		// We don't want to trigger selection changed event when the selection change is actually coming from the bar graph
+		bSuppressSelectionChangedEvent = true;
+
+		SelectedThread = Selection;
+		SelectedEvents = Selection.Get()->Children;
+		CreateSelectedEventsView();
+		SortEventsList();
+
+		bSuppressSelectionChangedEvent = false;
+	}
+}
+
+void SEventsTree::HandleBarEventSelectionChanged( int32 Thread, TSharedPtr<FVisualizerEvent> Selection )
+{
+	const TSharedPtr< FVisualizerEvent >* MappedViewSelection = ViewToEventsMap.FindKey( Selection );
+	if( MappedViewSelection )
+	{
+		TSharedPtr< FVisualizerEvent > ViewSelection = *MappedViewSelection;
+
+		// Clear the current selection
+		EventsListView->ClearSelection();
+		EventsListView->RequestTreeRefresh();
+
+		// Expand all parents so that it's visible
+		for( TSharedPtr< FVisualizerEvent > ParentEvent = ViewSelection->ParentEvent; ParentEvent.IsValid(); ParentEvent = ParentEvent->ParentEvent )
+		{
+			EventsListView->SetItemExpansion( ParentEvent, true );
+		}
+
+		EventsListView->SetSelection( ViewSelection );
+
+		EventsListView->RequestScrollIntoView( ViewSelection );
+	}
+	else
+	{
+		EventsListView->ClearSelection();
+	}
+}
+
+FString SEventsTree::GetTabTitle() const
+{
+	if ( SelectedThread.IsValid() )
+	{
+		return SelectedThread->EventName;
+	}
+	else
+	{
+		return NSLOCTEXT("ProfileVisualizer", "EventsVisualizerName", "Empty Events List").ToString();
+	}
+}
+
+int32 SEventsTree::CountEvents( TArray< TSharedPtr< FVisualizerEvent > >& Events )
+{
+	int32 Count = Events.Num();
+
+	for( int32 Index = 0; Index < Events.Num(); Index++ )
+	{
+		Count += CountEvents( Events[ Index ]->Children );
+	}
+
+	return Count;
+}
+
+void SEventsTree::CreateSelectedEventsView()
+{
+	const int32 EventsCount = CountEvents( SelectedEvents );
+	ViewToEventsMap.Empty( EventsCount );
+	SelectedEventsView.Empty( SelectedEvents.Num() );
+
+	// Create the selected events copy based on the current view mode
+	if( ViewMode == EVisualizerViewMode::Hierarchical )
+	{
+		SelectedEventsView.Empty( SelectedEvents.Num() );
+
+		for( int32 Index = 0; Index < SelectedEvents.Num(); Index++ )
+		{
+			if( FilterEvent( SelectedEvents[ Index ] ) )
+			{
+			TSharedPtr< FVisualizerEvent > EventCopy( CreateSelectedEventsViewRecursively( SelectedEvents[ Index ] ) );
+				if( EventCopy.IsValid() )
+				{
+			SelectedEventsView.Add( EventCopy );
+		}
+	}
+		}
+	}
+	else if( ViewMode == EVisualizerViewMode::Flat )
+	{
+		SelectedEventsView.Empty( EventsCount );
+
+		for( int32 Index = 0; Index < SelectedEvents.Num(); Index++ )
+		{
+			CreateSelectedEventsViewRecursivelyAndFlatten( SelectedEvents[ Index ] );
+		}
+	}
+	else if( ViewMode == EVisualizerViewMode::Coalesced )
+	{
+		SelectedEventsView.Empty( SelectedEvents.Num() );
+		CreateSelectedEventsViewRecursivelyCoalesced( SelectedEvents, SelectedEventsView, TSharedPtr< FVisualizerEvent >() );
+	}
+	else if( ViewMode == EVisualizerViewMode::FlatCoalesced )
+	{
+		SelectedEventsView.Empty( EventsCount );
+		CreateSelectedEventsViewRecursivelyFlatCoalesced( SelectedEvents );
+	}
+}
+
+TSharedPtr< FVisualizerEvent > SEventsTree::CreateSelectedEventsViewRecursively( TSharedPtr< FVisualizerEvent > SourceEvent )
+{
+	TSharedPtr< FVisualizerEvent > EventCopy( new FVisualizerEvent( *SourceEvent ) );
+
+	EventCopy->Children.Empty( SourceEvent->Children.Num() );
+	for( int32 ChildIndex = 0; ChildIndex < SourceEvent->Children.Num(); ChildIndex++ )
+	{
+		if( FilterEvent( SourceEvent->Children[ ChildIndex ] ) )
+		{
+		TSharedPtr< FVisualizerEvent > ChildCopy = CreateSelectedEventsViewRecursively( SourceEvent->Children[ ChildIndex ] );
+			if( ChildCopy.IsValid() )
+			{
+		ChildCopy->ParentEvent = EventCopy;
+		EventCopy->Children.Add( ChildCopy );
+	}
+		}
+	}
+
+	// Add this event because it's a leaf or has valid children
+	if( EventCopy->Children.Num() > 0 || SourceEvent->Children.Num() == 0 )
+	{
+		ViewToEventsMap.Add( EventCopy, SourceEvent );
+	}
+	else
+	{
+		EventCopy.Reset();
+	}
+
+	return EventCopy;
+}
+
+void SEventsTree::CreateSelectedEventsViewRecursivelyAndFlatten( TSharedPtr< FVisualizerEvent > SourceEvent )
+{
+	// Collect only children and store them directly into SelectedEventsView
+	if ( FilterEvent( SourceEvent ) && SourceEvent->Children.Num() == 0 )
+	{
+		TSharedPtr< FVisualizerEvent > EventCopy( new FVisualizerEvent( *SourceEvent ) );
+		EventCopy->ParentEvent.Reset();
+		ViewToEventsMap.Add( EventCopy, SourceEvent );
+		SelectedEventsView.Add( EventCopy );
+	}
+	
+	for( int32 ChildIndex = 0; ChildIndex < SourceEvent->Children.Num(); ChildIndex++ )
+	{
+		CreateSelectedEventsViewRecursivelyAndFlatten( SourceEvent->Children[ ChildIndex ] );
+	}
+}
+
+void SEventsTree::CreateSelectedEventsViewRecursivelyCoalesced( TArray< TSharedPtr< FVisualizerEvent > >& SourceEvents, TArray< TSharedPtr< FVisualizerEvent > >& CopiedEvents, TSharedPtr< FVisualizerEvent > InParent )
+{
+	for( int32 SourceIndex = 0; SourceIndex < SourceEvents.Num(); SourceIndex++ )
+	{
+		TSharedPtr< FVisualizerEvent > SourceEvent = SourceEvents[ SourceIndex ];
+		if( FilterEvent( SourceEvent ) )
+		{
+		if( SourceEvent->Children.Num() == 0 )
+		{
+			// Check if a child with the same name has already been added to the copied event
+			bool bEventExists = false;
+			for( int32 CopiedEventIndex = 0; CopiedEventIndex < CopiedEvents.Num(); CopiedEventIndex++ )
+			{
+				if( CopiedEvents[ CopiedEventIndex ]->EventName == SourceEvent->EventName )
+				{
+					bEventExists = true;
+					break;
+				}
+			}
+
+			if( !bEventExists )
+			{
+				TSharedPtr< FVisualizerEvent > EventCopy( new FVisualizerEvent( *SourceEvent ) );
+				EventCopy->ParentEvent = InParent;
+
+				ViewToEventsMap.Add( EventCopy, SourceEvent );
+				CopiedEvents.Add( EventCopy );				
+
+				// Find other children with the same name and add their time to the copied one
+				for( int32 OtherIndex = SourceIndex + 1; OtherIndex < SourceEvents.Num(); OtherIndex++ )
+				{
+					TSharedPtr< FVisualizerEvent > OtherEvent = SourceEvents[ OtherIndex ];
+					if( OtherEvent->Children.Num() == 0 && OtherEvent->EventName == SourceEvent->EventName )
+					{
+						EventCopy->DurationMs += OtherEvent->DurationMs;
+					}
+				}
+			}
+		}
+		else
+		{
+			TSharedPtr< FVisualizerEvent > EventCopy( new FVisualizerEvent( *SourceEvent ) );
+			EventCopy->Children.Empty( SourceEvent->Children.Num() );
+			EventCopy->ParentEvent = InParent;
+
+				CreateSelectedEventsViewRecursivelyCoalesced( SourceEvent->Children, EventCopy->Children, EventCopy );
+
+				// Only add this event if its children haven't been filtered
+				if( EventCopy->Children.Num() )
+				{
+			ViewToEventsMap.Add( EventCopy, SourceEvent );
+			CopiedEvents.Add( EventCopy );
+				}
+			}
+		}
+	}
+}
+
+void SEventsTree::CreateSelectedEventsViewRecursivelyFlatCoalesced( TArray< TSharedPtr< FVisualizerEvent > >& SourceEvents )
+{
+	for( int32 SourceIndex = 0; SourceIndex < SourceEvents.Num(); SourceIndex++ )
+	{
+		TSharedPtr< FVisualizerEvent > SourceEvent = SourceEvents[ SourceIndex ];
+		if( FilterEvent( SourceEvent ) )
+		{
+			if( SourceEvent->Children.Num() == 0 )
+			{
+				// Check if a child with the same name has already been added to the copied event
+				bool bEventExists = false;
+				for( int32 CopiedEventIndex = 0; CopiedEventIndex < SelectedEventsView.Num(); CopiedEventIndex++ )
+				{
+					if( SelectedEventsView[ CopiedEventIndex ]->EventName == SourceEvent->EventName )
+					{
+						bEventExists = true;
+						break;
+					}
+				}
+  
+				if( !bEventExists )
+				{
+					TSharedPtr< FVisualizerEvent > EventCopy( new FVisualizerEvent( *SourceEvent ) );
+					EventCopy->ParentEvent.Reset();
+  
+					ViewToEventsMap.Add( EventCopy, SourceEvent );
+					SelectedEventsView.Add( EventCopy );				
+  
+					// Find other children with the same name and add their time to the copied one
+					for( int32 OtherIndex = SourceIndex + 1; OtherIndex < SourceEvents.Num(); OtherIndex++ )
+					{
+						TSharedPtr< FVisualizerEvent > OtherEvent = SourceEvents[ OtherIndex ];
+						if( OtherEvent->Children.Num() == 0 && OtherEvent->EventName == SourceEvent->EventName )
+						{
+							EventCopy->DurationMs += OtherEvent->DurationMs;
+						}
+					}
+				}
+			}
+			else
+			{
+				CreateSelectedEventsViewRecursivelyFlatCoalesced( SourceEvent->Children );
+			}
+		}
+	}
+}
+
+void SEventsTree::OnFilterTextChanged( const FText& InFilterText )
+{
+	FilterText = InFilterText.ToString();
+	CreateSelectedEventsView();
+	EventsListView->RequestTreeRefresh();
+}
+
+void SEventsTree::OnFilterTextCommitted( const FText& InFilterText, ETextCommit::Type /*CommitInfo*/ )
+{
+}
+
+#undef LOCTEXT_NAMESPACE
diff --git a/Engine/Source/Developer/ProfileVisualizer/Private/SGraphBar.cpp b/Engine/Source/Developer/ProfileVisualizer/Private/SGraphBar.cpp
new file mode 100644
--- /dev/null
+++ b/Engine/Source/Developer/ProfileVisualizer/Private/SGraphBar.cpp
@@ -0,0 +1,217 @@
+// Copyright Epic Games, Inc. All Rights Reserved.
+
+#include "SGraphBar.h"
+#include "Rendering/DrawElements.h"
+#include "ProfileVisualizerStyle.h"
+
+FColor SGraphBar::ColorPalette[] = {
+	FColor( 0xff00A480 ),
+	FColor( 0xff62E200 ),
+	FColor( 0xff8F04A8 ),
+	FColor( 0xff1F7B67 ),
+	FColor( 0xff62AA2A ),
+	FColor( 0xff70227E ),
+	FColor( 0xff006B53 ),
+	FColor( 0xff409300 ),
+	FColor( 0xff5D016D ),
+	FColor( 0xff34D2AF ),
+	FColor( 0xff8BF13C ),
+	FColor( 0xffBC38D3 ),
+	FColor( 0xff5ED2B8 ),
+	FColor( 0xffA6F16C ),
+	FColor( 0xffC262D3 ),
+	FColor( 0xff0F4FA8 ),
+	FColor( 0xff00AE68 ),
+	FColor( 0xffDC0055 ),
+	FColor( 0xff284C7E ),
+	FColor( 0xff21825B ),
+	FColor( 0xffA52959 ),
+	FColor( 0xff05316D ),
+	FColor( 0xff007143 ),
+	FColor( 0xff8F0037 ),
+	FColor( 0xff4380D3 ),
+	FColor( 0xff36D695 ),
+	FColor( 0xffEE3B80 ),
+	FColor( 0xff6996D3 ),
+	FColor( 0xff60D6A7 ),
+	FColor( 0xffEE6B9E )
+};
+
+void SGraphBar::Construct( const FArguments& InArgs )
+{
+	OnSelectionChanged = InArgs._OnSelectionChanged;
+	OnGeometryChanged = InArgs._OnGeometryChanged;
+
+	BackgroundImage = FProfileVisualizerStyle::Get()->GetBrush("ProfileVisualizer.Background");
+	FillImage = FProfileVisualizerStyle::Get()->GetBrush("ProfileVisualizer.Mono");
+	SelectedImage = FProfileVisualizerStyle::Get()->GetBrush("ProfileVisualizer.Selected");
+
+	LastHoveredEvent = INDEX_NONE;
+	Zoom = 1.0f;
+	Offset = 0.0f;
+	StartTime = 0.0;
+	TotalTime = 1.0;
+}
+
+int32 SGraphBar::OnPaint( const FPaintArgs& Args, const FGeometry& AllottedGeometry, const FSlateRect& MyCullingRect, FSlateWindowElementList& OutDrawElements, int32 LayerId, const FWidgetStyle& InWidgetStyle, bool bParentEnabled ) const
+{
+	// Used to track the layer ID we will return.
+	int32 RetLayerId = LayerId;
+
+	bool bEnabled = ShouldBeEnabled( bParentEnabled );
+	const ESlateDrawEffect DrawEffects = bEnabled ? ESlateDrawEffect::None : ESlateDrawEffect::DisabledEffect;
+		
+	const FLinearColor ColorAndOpacitySRGB = InWidgetStyle.GetColorAndOpacityTint();
+	static const FLinearColor SelectedBarColor(FLinearColor::White);
+
+	// Paint inside the border only. 
+	const FVector2D BorderPadding = FProfileVisualizerStyle::Get()->GetVector("ProfileVisualizer.ProgressBar.BorderPadding");
+
+	FSlateDrawElement::MakeBox(
+		OutDrawElements,
+		RetLayerId++,
+		AllottedGeometry.ToPaintGeometry(),
+		BackgroundImage,
+		DrawEffects,
+		ColorAndOpacitySRGB
+	);	
+	
+	// Draw all bars
+	for( int32 EventIndex = 0; EventIndex < Events.Num(); ++EventIndex )
+	{
+		TSharedPtr< FVisualizerEvent > Event = Events[ EventIndex ];
+		float StartX, EndX;
+		if( CalculateEventGeometry( Event.Get(), AllottedGeometry, StartX, EndX ) )
+		{
+			// Draw Event bar
+			FSlateDrawElement::MakeBox(
+				OutDrawElements,
+				RetLayerId++,
+				AllottedGeometry.ToPaintGeometry(
+					FVector2D( StartX, 0.0f ),
+					FVector2D( EndX - StartX, AllottedGeometry.GetLocalSize().Y )),
+				Event->IsSelected ? SelectedImage : FillImage,
+				DrawEffects,
+				Event->IsSelected ? SelectedBarColor : ColorPalette[Event->ColorIndex % (sizeof(ColorPalette) / sizeof(ColorPalette[0]))]			
+				);
+		}
+	}
+
+	return RetLayerId - 1;
+}
+
+void SGraphBar::Tick( const FGeometry& AllottedGeometry, const double InCurrentTime, const float InDeltaTime )
+{
+	if( OnGeometryChanged.IsBound() )
+	{
+		if( AllottedGeometry != this->LastGeometry )
+		{
+			OnGeometryChanged.ExecuteIfBound( AllottedGeometry );
+			LastGeometry = AllottedGeometry;
+		}
+	}
+}
+
+FReply SGraphBar::OnMouseButtonDown( const FGeometry& MyGeometry, const FPointerEvent& MouseEvent )
+{
+	if( MouseEvent.GetEffectingButton() == EKeys::LeftMouseButton )
+	{
+		// Translate click position from absolute to graph space
+		const float ClickX = MyGeometry.AbsoluteToLocal( MouseEvent.GetScreenSpacePosition() ).X;//( MyGeometry.AbsoluteToLocal( MouseEvent.GetScreenSpacePosition() ).X / MyGeometry.Size.X ) / Zoom - Offset / Zoom;
+		TSharedPtr< FVisualizerEvent > SelectedEvent;
+
+		// Go through all events and check if at least one has been clicked
+		int32 EventIndex = INDEX_NONE;
+		for( EventIndex = 0; EventIndex < Events.Num(); EventIndex++ )
+		{
+			TSharedPtr< FVisualizerEvent > Event = Events[ EventIndex ];
+			float StartX, EndX;
+			if( CalculateEventGeometry( Event.Get(), MyGeometry, StartX, EndX ) )
+			{
+				if( ClickX >= StartX && ClickX < EndX )
+				{
+					// An event has been clicked, abort
+					SelectedEvent = Event;
+					Event->IsSelected = true;
+					break;
+				}
+				else
+				{
+					Event->IsSelected = false;
+				}
+			}
+		}
+
+		// Iterate over all the remaining events and mark them as not selected
+		for( EventIndex = EventIndex + 1; EventIndex < Events.Num(); EventIndex++ )
+		{
+			Events[ EventIndex ]->IsSelected = false;
+		}
+
+		// Execute OnSelectionChanged delegate
+		if( OnSelectionChanged.IsBound() )
+		{
+			OnSelectionChanged.ExecuteIfBound( SelectedEvent, ESelectInfo::OnMouseClick );
+		}
+
+		return FReply::Handled();
+	}
+	else
+	{
+		return FReply::Unhandled();
+	}
+}
+
+
+/**
+	* The system calls this method to notify the widget that a mouse moved within it. This event is bubbled.
+	*
+	* @param MyGeometry The Geometry of the widget receiving the event
+	* @param MouseEvent Information about the input event
+	*
+	* @return Whether the event was handled along with possible requests for the system to take action.
+	*/
+FReply SGraphBar::OnMouseMove( const FGeometry& MyGeometry, const FPointerEvent& MouseEvent )
+{
+	const float HoverX = ( MyGeometry.AbsoluteToLocal( MouseEvent.GetScreenSpacePosition() ).X / MyGeometry.GetLocalSize().X ) / Zoom - Offset / Zoom;
+
+	int32 HoveredEventIndex = INDEX_NONE;
+
+	for( int32 EventIndex = 0; EventIndex < Events.Num(); EventIndex++ )
+	{
+		TSharedPtr< FVisualizerEvent > Event = Events[ EventIndex ];
+
+		if( HoverX >= Event->Start && HoverX < ( Event->Start + Event->Duration ) )
+		{
+			HoveredEventIndex = EventIndex;
+			break;
+		}
+	}
+
+	if( HoveredEventIndex != LastHoveredEvent )
+	{
+		if( HoveredEventIndex != INDEX_NONE )
+		{
+			this->SetToolTipText( FText::FromString(Events[HoveredEventIndex]->EventName) );
+		}
+		else
+		{
+			this->SetToolTipText( FText::GetEmpty() );
+		}
+		LastHoveredEvent = HoveredEventIndex;
+	}
+
+	return SLeafWidget::OnMouseMove( MyGeometry, MouseEvent );
+}
+
+FVector2D SGraphBar::ComputeDesiredSize( float ) const
+{
+	return FVector2D( 8.0f, 16.0f );
+}
+
+void SGraphBar::SetEvents( const FVisualizerEventsArray& InEvents, double InStartTime, double InTotalTime )
+{
+	Events = InEvents;
+	StartTime = InStartTime;
+	TotalTime = InTotalTime;
+}
diff --git a/Engine/Source/Developer/ProfileVisualizer/Private/SProfileVisualizer.cpp b/Engine/Source/Developer/ProfileVisualizer/Private/SProfileVisualizer.cpp
new file mode 100644
--- /dev/null
+++ b/Engine/Source/Developer/ProfileVisualizer/Private/SProfileVisualizer.cpp
@@ -0,0 +1,321 @@
+// Copyright Epic Games, Inc. All Rights Reserved.
+
+#include "SProfileVisualizer.h"
+#include "ProfileVisualizerModule.h"
+#include "Widgets/SBoxPanel.h"
+#include "Layout/WidgetPath.h"
+#include "Framework/Application/MenuStack.h"
+#include "Framework/Application/SlateApplication.h"
+#include "Textures/SlateIcon.h"
+#include "Framework/Commands/UIAction.h"
+#include "Widgets/Layout/SBorder.h"
+#include "Widgets/Text/STextBlock.h"
+#include "Framework/MultiBox/MultiBoxBuilder.h"
+#include "Widgets/Layout/SSplitter.h"
+#include "ProfileVisualizerStyle.h"
+#include "SBarVisualizer.h"
+#include "SEventsTree.h"
+#include "SProfileVisualizer.h"
+#include "Widgets/Input/SButton.h"
+#include "Widgets/Text/STextBlock.h"
+#include "Misc/App.h"
+#include "Misc/EngineVersion.h"
+#include "HAL/FileManager.h"
+#include "Framework/Notifications/NotificationManager.h"
+#include "Widgets/Notifications/SNotificationList.h"
+
+#if PLATFORM_DESKTOP && WITH_EDITOR
+#include "DesktopPlatformModule.h"
+#endif
+
+#include "Framework/Commands/GenericCommands.h"
+	
+void SProfileVisualizer::Construct(const FArguments& InArgs)
+{
+	ProfileData = InArgs._ProfileData;
+	ProfilerType = InArgs._ProfilerType;
+	HeaderMessageText = InArgs._HeaderMessageText;
+	HeaderMessageTextColor = InArgs._HeaderMessageTextColor;
+
+	const FSlateBrush* ContentAreaBrush = FProfileVisualizerStyle::Get()->GetBrush("ProfileVisualizer.ContentAreaBrush");
+	ChildSlot
+	[
+		SNew(SVerticalBox)
+		+ SVerticalBox::Slot()
+		.AutoHeight()
+		[
+			SNew(SHorizontalBox)
+			+ SHorizontalBox::Slot()
+			.AutoWidth()
+			.Padding(0.0f)
+			[
+				SNew(SButton)
+				.OnClicked(this, &SProfileVisualizer::OnLoadClicked)
+				.Content()
+				[
+					SNew(STextBlock)
+					.Text(NSLOCTEXT("ProfileVisualizer", "Load", "Load"))
+					.ToolTipText(NSLOCTEXT("ProfileVisualizer", "Load_GPUTooltip", "Load GPU profiling data"))
+				]
+			]
+			+ SHorizontalBox::Slot()
+			.AutoWidth()
+			.Padding(0.0f)
+			[
+				SNew(SButton)
+				.OnClicked(this, &SProfileVisualizer::OnSaveClicked)
+				.Content()
+				[
+					SNew(STextBlock)
+					.Text(NSLOCTEXT("ProfileVisualizer", "Save", "Save"))
+					.ToolTipText(NSLOCTEXT("ProfileVisualizer", "Save_GPUTooltip", "Save the GPU profiling data"))
+				]
+			]
+		]
+		+ SVerticalBox::Slot()
+		.AutoHeight()
+		[
+			SNew(SBorder)
+			.Visibility(EVisibility::Visible)
+			.BorderImage(ContentAreaBrush)
+			[
+				SNew(STextBlock)
+				.Visibility( HeaderMessageText.IsEmpty() ? EVisibility::Collapsed : EVisibility::Visible )
+				.Text(HeaderMessageText)
+				.ColorAndOpacity(HeaderMessageTextColor)
+				.AutoWrapText(true)
+			]
+		]
+		+ SVerticalBox::Slot()
+		[
+			SAssignNew( MainSplitter, SSplitter )
+			.Orientation( Orient_Vertical )
+			+ SSplitter::Slot()
+			.Value(1.0f)
+			[
+				SNew( SBorder )
+				.Visibility( EVisibility::Visible )
+				.BorderImage( ContentAreaBrush )
+				[
+					SAssignNew( BarVisualizer, SBarVisualizer )
+					.ProfileData( ProfileData )
+					.OnBarGraphSelectionChanged( this, &SProfileVisualizer::RouteBarGraphSelectionChanged )
+					.OnBarGraphExpansionChanged( this, &SProfileVisualizer::RouteBarGraphExpansionChanged )
+					.OnBarEventSelectionChanged( this, &SProfileVisualizer::RouteBarEventSelectionChanged )
+					.OnBarGraphContextMenu( this, &SProfileVisualizer::OnBarGraphContextMenu )
+				]
+			]
+			+ SSplitter::Slot()
+			.Value(1.0f)
+			[
+				SNew( SBorder )
+				.Visibility( EVisibility::Visible )
+				.BorderImage( ContentAreaBrush )
+				[
+					SAssignNew( EventsTree, SEventsTree )
+					.ProfileData( ProfileData )
+					.OnEventSelectionChanged( this, &SProfileVisualizer::RouteEventSelectionChanged )
+				]
+			]
+		]
+	];
+
+	// Attempt to choose initial data set to display in the tree view
+	EventsTree->HandleBarGraphExpansionChanged( ProfileData );
+}
+
+TSharedRef< SWidget > SProfileVisualizer::MakeMainMenu()
+{
+	FMenuBarBuilder MenuBuilder( NULL );
+	{
+		// File
+		MenuBuilder.AddPullDownMenu( 
+			NSLOCTEXT("ProfileVisualizer", "FileMenu", "File"),
+			NSLOCTEXT("ProfileVisualizer", "FileMenu_ToolTip", "Open the file menu"),
+			FNewMenuDelegate::CreateSP( this, &SProfileVisualizer::FillFileMenu ) );
+
+		// Apps
+		MenuBuilder.AddPullDownMenu( 
+			NSLOCTEXT("ProfileVisualizer", "AppMenu", "Window"),
+			NSLOCTEXT("ProfileVisualizer", "AppMenu_ToolTip", "Open the summoning menu"),
+			FNewMenuDelegate::CreateSP( this, &SProfileVisualizer::FillAppMenu ) );
+
+		// Help
+		MenuBuilder.AddPullDownMenu( 
+			NSLOCTEXT("ProfileVisualizer", "HelpMenu", "Help"),
+			NSLOCTEXT("ProfileVisualizer", "HelpMenu_ToolTip", "Open the help menu"),
+			FNewMenuDelegate::CreateSP( this, &SProfileVisualizer::FillHelpMenu ) );
+	}
+
+	// Create the menu bar
+	TSharedRef< SWidget > MenuBarWidget = MenuBuilder.MakeWidget();
+
+	return MenuBarWidget;
+}
+
+void SProfileVisualizer::FillFileMenu( FMenuBuilder& MenuBuilder )
+{
+}
+
+void SProfileVisualizer::FillAppMenu( FMenuBuilder& MenuBuilder )
+{
+}
+
+void SProfileVisualizer::FillHelpMenu( FMenuBuilder& MenuBuilder )
+{
+}
+
+
+void SProfileVisualizer::RouteEventSelectionChanged( TSharedPtr< FVisualizerEvent > Selection )
+{
+	BarVisualizer->HandleEventSelectionChanged( Selection );
+}
+
+void SProfileVisualizer::RouteBarGraphSelectionChanged( TSharedPtr< FVisualizerEvent > Selection )
+{
+	EventsTree->HandleBarGraphSelectionChanged( Selection );
+}
+
+void SProfileVisualizer::RouteBarGraphExpansionChanged( TSharedPtr< FVisualizerEvent > Selection )
+{
+	EventsTree->HandleBarGraphExpansionChanged( Selection );
+}
+
+void SProfileVisualizer::RouteBarEventSelectionChanged( int32 Thread, TSharedPtr<FVisualizerEvent> Selection )
+{
+	EventsTree->HandleBarEventSelectionChanged( Thread, Selection );
+}
+
+void SProfileVisualizer::OnBarGraphContextMenu( TSharedPtr< FVisualizerEvent > Selection, const FPointerEvent& InputEvent )
+{
+	SelectedBarGraph = Selection;
+
+	FWidgetPath WidgetPath = InputEvent.GetEventPath() != nullptr ? *InputEvent.GetEventPath() : FWidgetPath();
+	FSlateApplication::Get().PushMenu(SharedThis(this), WidgetPath, MakeBarVisualizerContextMenu(), FSlateApplication::Get().GetCursorPos(), FPopupTransitionEffect::ContextMenu);
+}
+
+TSharedRef<SWidget> SProfileVisualizer::MakeBarVisualizerContextMenu()
+{
+	const bool bCloseAfterSelection = true;
+	FMenuBuilder MenuBuilder( bCloseAfterSelection, NULL );
+	{
+		FUIAction Action( FExecuteAction::CreateSP( this, &SProfileVisualizer::ShowGraphBarInEventsWindow, (int32)INDEX_NONE ) );
+		MenuBuilder.AddMenuEntry( NSLOCTEXT("ProfileVisualizer", "GraphBarShowInNew", "Show in New Events Window"), FText::GetEmpty(), FSlateIcon(), Action, NAME_None, EUserInterfaceActionType::Button );
+	}
+		
+	return MenuBuilder.MakeWidget();
+}
+
+void SProfileVisualizer::ShowGraphBarInEventsWindow( int32 WindowIndex )
+{
+	EventsTree->HandleBarGraphExpansionChanged( SelectedBarGraph );
+}
+
+FReply SProfileVisualizer::OnSaveClicked()
+{
+#if PLATFORM_DESKTOP && WITH_EDITOR
+	// Message to display on completion
+	FText Message;
+
+	FString ProfileFilename = FPaths::ProjectLogDir();
+	ProfileFilename /= TEXT("profileViz");
+	ProfileFilename /= FString::Printf(
+		TEXT("%s-%i-%s.profViz"),
+		FApp::GetProjectName(),
+		FEngineVersion::Current().GetChangelist(),
+		*FDateTime::Now().ToString());
+
+	FArchive* ProfileFile = IFileManager::Get().CreateFileWriter(*ProfileFilename);
+
+	if (ProfileFile)
+	{
+		FVisualizerEvent::SaveVisualizerEventRecursively(ProfileFile, ProfileData);
+
+		// Close and delete archive.
+		ProfileFile->Close();
+		delete ProfileFile;
+		ProfileFile = NULL;
+
+		Message = NSLOCTEXT("ProfileVisualizer", "ExportMessage", "Wrote profile data to file");
+	}
+	else
+	{
+		Message = NSLOCTEXT("ProfileVisualizer", "FailedExportMessage", "Could not write profile data to file");
+	}
+
+	struct Local
+	{
+		static void NavigateToExportedFile(FString InGPUFilename, bool bInSuccessful)
+		{
+			InGPUFilename = FPaths::ConvertRelativePathToFull(InGPUFilename);
+			if (bInSuccessful)
+			{
+				FPlatformProcess::LaunchFileInDefaultExternalApplication(*InGPUFilename);
+			}
+			else
+			{
+				FPlatformProcess::ExploreFolder(*FPaths::GetPath(InGPUFilename));
+			}
+		}
+	};
+
+	FNotificationInfo Info(Message);
+	Info.Hyperlink = FSimpleDelegate::CreateStatic(&Local::NavigateToExportedFile, ProfileFilename, false);
+	Info.HyperlinkText = FText::FromString(ProfileFilename);
+	Info.bUseLargeFont = false;
+	Info.bFireAndForget = true;
+	Info.ExpireDuration = 8.0f;
+	FSlateNotificationManager::Get().AddNotification(Info);
+#endif //#if PLATFORM_DESKTOP && WITH_EDITOR
+
+	return FReply::Handled();
+
+}
+
+FReply SProfileVisualizer::OnLoadClicked()
+{
+#if PLATFORM_DESKTOP && WITH_EDITOR
+	// Prompt the user for the Filename
+	IDesktopPlatform* DesktopPlatform = FDesktopPlatformModule::Get();
+	if (DesktopPlatform != nullptr)
+	{
+		TArray<FString> OpenFilenames;
+		bool bOpened = DesktopPlatform->OpenFileDialog(
+			FSlateApplication::Get().FindBestParentWindowHandleForDialogs(nullptr),
+			NSLOCTEXT("UnrealEd", "Load", "Load Profile data").ToString(),
+			TEXT(""),
+			TEXT(""),
+			TEXT("Profile data (*.profViz) | *.profViz"),
+			EFileDialogFlags::None,
+			OpenFilenames);
+
+		if (bOpened)
+		{
+			FArchive* ProfileFile = IFileManager::Get().CreateFileReader(*OpenFilenames[0]);
+
+			if (ProfileFile)
+			{
+				TSharedPtr< FVisualizerEvent > InVisualizerData;
+				InVisualizerData = FVisualizerEvent::LoadVisualizerEvent(ProfileFile);
+								
+				static FName ProfileVisualizerModule(TEXT("ProfileVisualizer"));
+				if (FModuleManager::Get().IsModuleLoaded(ProfileVisualizerModule))
+				{
+					IProfileVisualizerModule& ProfileVisualizer = FModuleManager::GetModuleChecked<IProfileVisualizerModule>(ProfileVisualizerModule);
+
+					FText LoadedFileName = FText::AsCultureInvariant(ProfileFile->GetArchiveName());
+					ProfileVisualizer.DisplayProfileVisualizer(InVisualizerData, TEXT("Profile Data"), LoadedFileName);
+				}
+
+				// Close and delete archive.
+				ProfileFile->Close();
+				delete ProfileFile;
+				ProfileFile = NULL;
+			}
+		}
+	}
+#endif//#if PLATFORM_DESKTOP && WITH_EDITOR
+
+	return FReply::Handled();
+
+}
diff --git a/Engine/Source/Developer/ProfileVisualizer/Private/SProfileVisualizer.h b/Engine/Source/Developer/ProfileVisualizer/Private/SProfileVisualizer.h
new file mode 100644
--- /dev/null
+++ b/Engine/Source/Developer/ProfileVisualizer/Private/SProfileVisualizer.h
@@ -0,0 +1,111 @@
+// Copyright Epic Games, Inc. All Rights Reserved.
+
+#pragma once
+
+#include "CoreMinimal.h"
+#include "Widgets/SWidget.h"
+#include "Widgets/DeclarativeSyntaxSupport.h"
+#include "Widgets/SCompoundWidget.h"
+#include "VisualizerEvents.h"
+
+class FMenuBuilder;
+class SBarVisualizer;
+class SEventsTree;
+class SSplitter;
+
+class SProfileVisualizer : public SCompoundWidget
+{
+public:
+
+	SLATE_BEGIN_ARGS(SProfileVisualizer)
+		: _LayoutString()
+		, _ProfileData()
+		, _ProfilerType()
+		{}
+
+		SLATE_ARGUMENT( FString, LayoutString )
+
+		SLATE_ARGUMENT( TSharedPtr< FVisualizerEvent >, ProfileData )
+
+		SLATE_ARGUMENT( FText, ProfilerType )
+
+		SLATE_ARGUMENT( FText, HeaderMessageText )
+
+		SLATE_ARGUMENT( FLinearColor, HeaderMessageTextColor )
+
+	SLATE_END_ARGS()
+
+	
+	/**
+	 * Constructs this widget
+	 *
+	 * @param InArgs    Declaration from which to construct the widget
+	 */
+	void Construct( const FArguments& InArgs );
+
+private:
+
+	/** Construct main menu */
+	TSharedRef< SWidget > MakeMainMenu();
+
+	/** Fills the contents of 'File' menu */
+	void FillFileMenu( FMenuBuilder& MenuBuilder );
+
+	/** Fills the contents of 'Window' menu */
+	void FillAppMenu( FMenuBuilder& MenuBuilder );
+
+	/** Fills the contents of 'Help' menu */
+	void FillHelpMenu( FMenuBuilder& MenuBuilder );
+
+	/** Routes events from events tree tab to bar visualizer tabs */
+	void RouteEventSelectionChanged( TSharedPtr< FVisualizerEvent > Selection );
+
+	/** Routes events from bar visualizer tab to events tree tabs */
+	void RouteBarGraphSelectionChanged( TSharedPtr< FVisualizerEvent > Selection );
+
+	/** Routes events from bar visualizer tab to events tree tabs */
+	void RouteBarGraphExpansionChanged( TSharedPtr< FVisualizerEvent > Selection );
+
+	/** Routes events from bar visualizer tab to events tree tabs */
+	void RouteBarEventSelectionChanged( int32 Thread, TSharedPtr<FVisualizerEvent> Selection );
+
+	/** Opens a context menu when a bar is right clicked */
+	void OnBarGraphContextMenu( TSharedPtr< FVisualizerEvent > Selection, const FPointerEvent& InputEvent );
+
+	/** Creates bar visualizer context menu for opening new events tree tabs */
+	TSharedRef<SWidget> MakeBarVisualizerContextMenu();
+
+	/** Handles creating a new events tree tab from the bar visualizer context menu */
+	void ShowGraphBarInEventsWindow( int32 WindowIndex );
+
+	/** Pointer to the main dock area of this widget */
+	TSharedPtr< SSplitter > MainSplitter;
+
+	/** Pointer to the bar visualizer */
+	TSharedPtr< SBarVisualizer > BarVisualizer;
+
+	/** Pointer to the events visualizer */
+	TSharedPtr< SEventsTree > EventsTree;
+
+	/** Profiler results */
+	TSharedPtr< FVisualizerEvent > ProfileData;
+
+	/** Currently selected bar graph */
+	TSharedPtr< FVisualizerEvent > SelectedBarGraph;
+
+	/** Profiler name */
+	FText ProfilerType;
+
+	/** Optional header message to display at the top of the profile window */
+	FText HeaderMessageText;
+	 
+	/** Optional header message text color */
+	FLinearColor HeaderMessageTextColor;
+
+	/** We just clicked the save button */
+	FReply OnSaveClicked();
+
+	/** We just clicked the load button */
+	FReply OnLoadClicked();
+};
+
diff --git a/Engine/Source/Developer/ProfileVisualizer/Private/STimeline.cpp b/Engine/Source/Developer/ProfileVisualizer/Private/STimeline.cpp
new file mode 100644
--- /dev/null
+++ b/Engine/Source/Developer/ProfileVisualizer/Private/STimeline.cpp
@@ -0,0 +1,170 @@
+// Copyright Epic Games, Inc. All Rights Reserved.
+
+#include "STimeline.h"
+#include "Misc/Paths.h"
+#include "Fonts/SlateFontInfo.h"
+#include "Fonts/FontMeasure.h"
+#include "Styling/CoreStyle.h"
+#include "Rendering/DrawElements.h"
+#include "Framework/Application/SlateApplication.h"
+#include "ProfileVisualizerStyle.h"
+
+void STimeline::Construct( const FArguments& InArgs )
+{
+	MinValue = InArgs._MinValue;
+	MaxValue = InArgs._MaxValue;
+	FixedLabelSpacing = InArgs._FixedLabelSpacing;
+
+	BackgroundImage = FProfileVisualizerStyle::Get()->GetBrush("ProfileVisualizer.Background");
+
+	Zoom = 1.0f;
+	Offset = 0.0f;
+}
+
+int32 STimeline::OnPaint( const FPaintArgs& Args, const FGeometry& AllottedGeometry, const FSlateRect& MyCullingRect, FSlateWindowElementList& OutDrawElements, int32 LayerId, const FWidgetStyle& InWidgetStyle, bool bParentEnabled ) const
+{
+	// Used to track the layer ID we will return.
+	int32 RetLayerId = LayerId;
+	const TSharedRef< FSlateFontMeasure > FontMeasureService = FSlateApplication::Get().GetRenderer()->GetFontMeasureService();
+
+	bool bEnabled = ShouldBeEnabled( bParentEnabled );
+	const ESlateDrawEffect DrawEffects = bEnabled ? ESlateDrawEffect::None : ESlateDrawEffect::DisabledEffect;
+		
+	const FLinearColor ColorAndOpacitySRGB = InWidgetStyle.GetColorAndOpacityTint();
+	static const FLinearColor SelectedBarColor(FLinearColor::White);
+
+	// Paint inside the border only. 
+	const FVector2D BorderPadding = FProfileVisualizerStyle::Get()->GetVector("ProfileVisualizer.ProgressBar.BorderPadding");
+	
+	const float OffsetX = DrawingOffsetX; // BorderPadding.X
+	const float Width = DrawingGeometry.Size.X; // AllottedGeometry.Size.X - - 2.0f * BorderPadding.X
+
+	FSlateFontInfo MyFont = FCoreStyle::GetDefaultFontStyle("Regular", 10);
+
+	// Create line points
+	const float RoundedMax = FMath::CeilToInt( MaxValue );
+	const float RoundedMin = FMath::FloorToInt( MinValue );
+	const float TimeScale = MaxValue - MinValue;
+	const int32 NumValues = FMath::FloorToInt( AllottedGeometry.Size.X * Zoom / FixedLabelSpacing );
+
+	TArray< FVector2D > LinePoints;
+	LinePoints.AddUninitialized( 2 );
+	
+	{
+		LinePoints[0] = FVector2D( OffsetX, BorderPadding.Y + 1.0f );
+		LinePoints[1] = FVector2D( OffsetX + Width, BorderPadding.Y + 1.0f );
+
+		// Draw lines
+		FSlateDrawElement::MakeLines( 
+			OutDrawElements,
+			RetLayerId++,
+			AllottedGeometry.ToPaintGeometry(),
+			LinePoints,
+			ESlateDrawEffect::None,
+			FLinearColor::White
+			);
+	}
+
+	const FVector2D TextDrawSize = FontMeasureService->Measure(TEXT("0.00"), MyFont);
+	const float LineHeight = AllottedGeometry.Size.Y - BorderPadding.Y - TextDrawSize.Y - 2.0f;
+	
+	for( int32 LineIndex = 0; LineIndex <= NumValues; LineIndex++ )
+	{
+		const float NormalizedX = (float)LineIndex / NumValues;
+		const float LineX = Offset + NormalizedX * Zoom;
+		if( LineX < 0.0f || LineX > 1.0f )
+		{
+			continue;
+		}
+
+		const float LineXPos =  OffsetX + Width * LineX;
+		LinePoints[0] = FVector2D( LineXPos, BorderPadding.Y );
+		LinePoints[1] = FVector2D( LineXPos, LineHeight );
+
+		// Draw lines
+		FSlateDrawElement::MakeLines( 
+			OutDrawElements,
+			RetLayerId++,
+			AllottedGeometry.ToPaintGeometry(),
+			LinePoints,
+			ESlateDrawEffect::None,
+			FLinearColor::White
+			);
+
+		FString ValueText( FString::Printf( TEXT("%.2f"), MinValue + NormalizedX * TimeScale ) );
+		FVector2D DrawSize = FontMeasureService->Measure(ValueText, MyFont);
+		FVector2D TextPos( LineXPos - DrawSize.X * 0.5f, LineHeight );
+
+		if( TextPos.X < 0.0f )
+		{
+			TextPos.X = 0.0f;
+		}
+		else if( (TextPos.X + DrawSize.X) > AllottedGeometry.Size.X )
+		{
+			TextPos.X = OffsetX + Width - DrawSize.X;
+		}
+
+		FSlateDrawElement::MakeText( 
+			OutDrawElements,
+			RetLayerId,
+			AllottedGeometry.ToOffsetPaintGeometry( TextPos ),
+			ValueText,
+			MyFont,
+			ESlateDrawEffect::None,
+			FLinearColor::White
+			);
+	}
+
+	// Always draw lines at the start and at the end of the timeline
+	{
+		LinePoints[0] = FVector2D( OffsetX, BorderPadding.Y );
+		LinePoints[1] = FVector2D( OffsetX, LineHeight );
+
+		// Draw lines
+		FSlateDrawElement::MakeLines( 
+			OutDrawElements,
+			RetLayerId++,
+			AllottedGeometry.ToPaintGeometry(),
+			LinePoints,
+			ESlateDrawEffect::None,
+			FLinearColor::White
+			);
+	}
+
+	{
+		LinePoints[0] = FVector2D( OffsetX + Width, BorderPadding.Y );
+		LinePoints[1] = FVector2D( OffsetX + Width, LineHeight );
+
+		// Draw lines
+		FSlateDrawElement::MakeLines( 
+			OutDrawElements,
+			RetLayerId++,
+			AllottedGeometry.ToPaintGeometry(),
+			LinePoints,
+			ESlateDrawEffect::None,
+			FLinearColor::White
+			);
+	}
+
+	return RetLayerId - 1;
+}
+
+void STimeline::Tick( const FGeometry& AllottedGeometry, const double InCurrentTime, const float InDeltaTime )
+{
+	DrawingOffsetX = DrawingGeometry.AbsolutePosition.X - AllottedGeometry.AbsolutePosition.X;
+}
+
+FReply STimeline::OnMouseButtonDown( const FGeometry& MyGeometry, const FPointerEvent& MouseEvent )
+{
+	return SCompoundWidget::OnMouseButtonDown( MyGeometry, MouseEvent );
+}
+
+FReply STimeline::OnMouseMove( const FGeometry& MyGeometry, const FPointerEvent& MouseEvent )
+{
+	return SCompoundWidget::OnMouseMove( MyGeometry, MouseEvent );
+}
+
+FVector2D STimeline::ComputeDesiredSize( float ) const
+{
+	return FVector2D( 8.0f, 8.0f );
+}
diff --git a/Engine/Source/Developer/ProfileVisualizer/Private/VisualizerEvents.cpp b/Engine/Source/Developer/ProfileVisualizer/Private/VisualizerEvents.cpp
new file mode 100644
--- /dev/null
+++ b/Engine/Source/Developer/ProfileVisualizer/Private/VisualizerEvents.cpp
@@ -0,0 +1,79 @@
+// Copyright Epic Games, Inc. All Rights Reserved.
+
+#include "VisualizerEvents.h"
+
+#if defined( _WIN64 )
+static_assert(sizeof(FVisualizerEvent) == 88, "FVisualizerEvent size has changed. Please update the save/load functions");
+#endif
+
+TSharedPtr< FVisualizerEvent > FVisualizerEvent::LoadVisualizerEventRecursively(FArchive *Ar, TSharedPtr< FVisualizerEvent > InParentEvent)
+{
+	double Start;
+	double Duration;
+	double DurationMs;
+	int32 Category;
+	FString EventName;
+	bool IsSelected;
+
+	*Ar << Start;
+	*Ar << Duration;
+	*Ar << DurationMs;
+	*Ar << Category;
+	*Ar << EventName;
+	*Ar << IsSelected;
+
+	TSharedPtr< FVisualizerEvent > VisualizerEvent(new FVisualizerEvent(Start, Duration,DurationMs,Category,EventName));
+	VisualizerEvent->ParentEvent = InParentEvent;
+
+	uint32 NumChildren;
+
+	Ar->SerializeInt(NumChildren, MAX_uint32);
+	VisualizerEvent->Children.Reserve(NumChildren);
+
+	for (uint32 i = 0; i < NumChildren; ++i)
+	{
+		TSharedPtr< FVisualizerEvent > ChildEvent = LoadVisualizerEventRecursively(Ar, VisualizerEvent);
+		VisualizerEvent->Children.Add(ChildEvent);
+	}
+
+	return VisualizerEvent;
+
+}
+
+TSharedPtr< FVisualizerEvent > FVisualizerEvent::LoadVisualizerEvent(FArchive *Ar)
+{
+	//FArchive Ar;
+	// Assumption: InProfileData contains only one (root) element. Otherwise an extra FVisualizerEvent root event is required.
+	TSharedPtr< FVisualizerEvent > DummyRoot;
+	// Recursively create visualizer event data.
+	TSharedPtr< FVisualizerEvent > StatEvents(LoadVisualizerEventRecursively( Ar, DummyRoot));
+	return StatEvents;
+}
+
+
+
+void FVisualizerEvent::SaveVisualizerEventRecursively(FArchive *Ar, TSharedPtr< FVisualizerEvent > VisualizerEvent )
+{
+	//Write all simple members
+	*Ar << VisualizerEvent->Start;
+	*Ar << VisualizerEvent->Duration;
+	*Ar << VisualizerEvent->DurationMs;
+	*Ar << VisualizerEvent->Category;
+	*Ar << VisualizerEvent->EventName;
+	*Ar << VisualizerEvent->IsSelected;
+
+	uint32 NumChildren;
+
+	NumChildren = VisualizerEvent->Children.Num();
+	Ar->SerializeInt(NumChildren, MAX_uint32);
+
+	for(uint32 i = 0; i < NumChildren; ++i)
+	{
+		TSharedPtr< FVisualizerEvent > Child;
+
+		Child = VisualizerEvent->Children[i];
+		SaveVisualizerEventRecursively(Ar, Child);
+	}
+}
+
+
diff --git a/Engine/Source/Developer/ProfileVisualizer/ProfileVisualizer.Build.cs b/Engine/Source/Developer/ProfileVisualizer/ProfileVisualizer.Build.cs
new file mode 100644
--- /dev/null
+++ b/Engine/Source/Developer/ProfileVisualizer/ProfileVisualizer.Build.cs
@@ -0,0 +1,36 @@
+// Copyright Epic Games, Inc. All Rights Reserved.
+
+using UnrealBuildTool;
+
+public class ProfileVisualizer : ModuleRules
+{
+	public ProfileVisualizer(ReadOnlyTargetRules Target) : base(Target)
+	{
+		PrivateDependencyModuleNames.AddRange(
+			new string[] {
+				"Core",
+				"CoreUObject",
+				"Slate",
+				"SlateCore",
+				
+				"Engine",
+				"InputCore"
+			}
+		);
+
+		if (Target.Type == TargetType.Editor)
+		{
+			PrivateDependencyModuleNames.AddRange(
+				new string[] {
+					"WorkspaceMenuStructure",
+					"DesktopPlatform"
+				}
+			);
+		}
+
+		if (Target.Configuration != UnrealTargetConfiguration.Shipping)
+		{
+			PrecompileForTargets = PrecompileTargetsType.Any;
+		}
+	}
+}
diff --git a/Engine/Source/Developer/ProfileVisualizer/Public/ProfileVisualizerModule.h b/Engine/Source/Developer/ProfileVisualizer/Public/ProfileVisualizerModule.h
new file mode 100644
--- /dev/null
+++ b/Engine/Source/Developer/ProfileVisualizer/Public/ProfileVisualizerModule.h
@@ -0,0 +1,14 @@
+// Copyright Epic Games, Inc. All Rights Reserved.
+
+#pragma once
+
+#include "CoreMinimal.h"
+#include "Modules/ModuleInterface.h"
+#include "VisualizerEvents.h"
+
+class IProfileVisualizerModule
+	: public IModuleInterface
+{
+public:
+	virtual void DisplayProfileVisualizer( TSharedPtr< FVisualizerEvent > InProfileData, const TCHAR* InProfilerType, const FText& HeaderMessageText = FText::GetEmpty(), const FLinearColor& HeaderMessageTextColor = FLinearColor::White ) = 0;
+};
diff --git a/Engine/Source/Developer/ProfileVisualizer/Public/SBarVisualizer.h b/Engine/Source/Developer/ProfileVisualizer/Public/SBarVisualizer.h
new file mode 100644
--- /dev/null
+++ b/Engine/Source/Developer/ProfileVisualizer/Public/SBarVisualizer.h
@@ -0,0 +1,289 @@
+// Copyright Epic Games, Inc. All Rights Reserved.
+
+#pragma once
+
+#include "CoreMinimal.h"
+#include "Layout/Visibility.h"
+#include "Layout/Geometry.h"
+#include "Input/Reply.h"
+#include "Widgets/DeclarativeSyntaxSupport.h"
+#include "Widgets/SCompoundWidget.h"
+#include "VisualizerEvents.h"
+#include "Widgets/Views/STableViewBase.h"
+#include "Widgets/Views/STableRow.h"
+
+class SGraphBar;
+class SScrollBar;
+class STextBlock;
+class STimeline;
+
+/**
+ * Bars Visualizer. Contains a list of bars for each profiler category
+ */
+class PROFILEVISUALIZER_API SBarVisualizer : public SCompoundWidget
+{
+public:
+
+	/** Delegate used when bar graph selection changes */
+	DECLARE_DELEGATE_OneParam( FOnBarGraphSelectionChanged, TSharedPtr< FVisualizerEvent > );
+
+	/** Delegate used when bar graph expansion changes */
+	DECLARE_DELEGATE_OneParam( FOnBarGraphExpansionChanged, TSharedPtr< FVisualizerEvent > );
+
+	/** Delegate used when a single event on the bar graph is selected */
+	DECLARE_DELEGATE_TwoParams( FOnBarEventSelectionChanged, int32, TSharedPtr< FVisualizerEvent > );
+	
+	/** Delegate used when the user right-clicks on a bar graph */
+	DECLARE_DELEGATE_TwoParams( FOnBarGraphContextMenu, TSharedPtr< FVisualizerEvent >, const FPointerEvent& );
+
+	SLATE_BEGIN_ARGS( SBarVisualizer )
+		: _ProfileData()
+		, _OnBarGraphSelectionChanged()
+		, _OnBarGraphExpansionChanged()
+		, _OnBarEventSelectionChanged()
+		, _OnBarGraphContextMenu()
+		{}
+
+		/** Profiler results */
+		SLATE_ATTRIBUTE( TSharedPtr< FVisualizerEvent >, ProfileData )
+
+		/** Callback triggered when bar graph selection changes */
+		SLATE_EVENT( FOnBarGraphSelectionChanged, OnBarGraphSelectionChanged )
+
+		/** Callback triggered when bar graph expansion changes */
+		SLATE_EVENT( FOnBarGraphExpansionChanged, OnBarGraphExpansionChanged )
+
+		/** Callback triggered when single event on the bar graph is selected */
+		SLATE_EVENT( FOnBarEventSelectionChanged, OnBarEventSelectionChanged )
+
+		/** Callback triggered when the user right-clicks on a bar graph */
+		SLATE_EVENT( FOnBarGraphContextMenu, OnBarGraphContextMenu )
+
+	SLATE_END_ARGS()
+
+	/**
+	 * Construct the widget
+	 * 
+	 * @param InArgs   A declaration from which to construct the widget
+	 */
+	void Construct( const FArguments& InArgs );
+
+	/**
+	 * Handles selection change in the events tree
+	 *
+	 * @param Selection selected event
+	 */
+	void HandleEventSelectionChanged( TSharedPtr< FVisualizerEvent > Selection );
+
+protected:
+
+	/**
+	 * Gets the maximum scroll offset fraction value for the horizontal scrollbar
+	 *
+	 * @return Maximum scroll offset fraction
+	 */
+	float GetMaxScrollOffsetFraction() const
+	{
+		return 1.0f - 1.0f / GetZoom();
+	}
+
+	/**
+	 * Gets the maximum graph offset value for the graph bars
+	 *
+	 * @return Maximum graph offset
+	 */
+	float GetMaxGraphOffset() const
+	{
+		return GetZoom() - 1.0f;
+	}
+
+	/**
+	 * Gets the actual zoom level for the graph bars
+	 *
+	 * @return Zoom level
+	 */
+	float GetZoom() const
+	{
+		const float MinZoom = 1.0f;
+		const float MaxZoom = 20.0f;
+		return MinZoom + ZoomSliderValue * ( MaxZoom - MinZoom );
+	}
+
+	/**
+	 * Callback for scrolling the horizontal scrollbar
+	 *
+	 * @param InScrollOffsetFraction Scrollbar offset fraction
+	 */
+	void ScrollBar_OnUserScrolled( float InScrollOffsetFraction );
+
+	/**
+	 * Constructs the zoom label string based on the current zoom level value.
+	 *
+	 * @return Zoom label text
+	 */
+	FText GetZoomLabel() const;
+
+	/**
+	 * Callback used to get the current zoom slider value.
+	 *
+	 * @return Zoom slider value
+	 */
+	float GetZoomValue() const;
+
+	/**
+	 * Callback used to handle zoom slider
+	 *
+	 * @param NewValue New Zoom slider value
+	 */
+	void OnSetZoomValue( float NewValue );
+
+	/**
+	 * Function called when the currently selected bar graph changes
+	 *
+	 * @param Selection Currently selected thread events
+	 * @param SelectInfo Provides context on how the selection changed
+	 */
+	void OnBarGraphSelectionChanged( TSharedPtr< FVisualizerEvent > Selection, ESelectInfo::Type SelectInfo );
+
+	/**
+	 * Function called when the user selects an event bar in the graph
+	 *
+	 * @param Selection Currently selected event
+	 * @param SelectInfo Provides context on how the selection changed
+	 */
+	void OnBarEventSelectionChanged( TSharedPtr< FVisualizerEvent > Selection, ESelectInfo::Type SelectInfo, int32 BarId );
+
+	/** 
+	 * Generates SGraphBar widget for the threads list
+	 *
+	 * @param InItem Graph profile events
+	 * @param OwnerTable Onwer Table
+	 */
+	TSharedRef<ITableRow> OnGenerateWidgetForList( TSharedPtr< FVisualizerEvent > InItem, const TSharedRef< STableViewBase >& OwnerTable );
+
+	/** 
+	 * Generates children for the specified tree view item
+	 *
+	 * @param InItem Graph profile events
+	 * @param OutChildren child graphs
+	 */
+	void OnGetChildrenForList( TSharedPtr<FVisualizerEvent> InItem, TArray<TSharedPtr<FVisualizerEvent> >& OutChildren);
+
+	/**
+	 * Forwards right-click event to the visualizer main frame
+	 *
+	 * @param BarGeometry Graph bar geometry
+	 * @param MouseEvent Current mouse event
+	 * @param Selection Selected events
+	 */
+	FReply OnBarRightClicked( const FGeometry& BarGeometry, const FPointerEvent& MouseEvent, TSharedPtr<FVisualizerEvent> Selection );
+
+	/**
+	 * Recursively clears selection on all bar graphs
+	 *
+	 * @param GraphEvents Bar Graph events
+	 * @param Selection Currently selected event
+	 */
+	void ClearBarSelection( TSharedPtr< FVisualizerEvent > GraphEvents, TSharedPtr<FVisualizerEvent> Selection );
+
+	/** Creates filtered data */
+	void CreateDataView();
+
+	/** Creates flattened data view */
+	void CreateFlattenedData( TSharedPtr< FVisualizerEvent > InData, TArray< TSharedPtr< FVisualizerEvent > >& FlattenedData );
+
+	/** 
+	 * Sets the current view mode 
+	 *
+	 * @param InMode New view mode
+	 */
+	void SetViewMode( EVisualizerViewMode::Type InMode );
+
+	/** 
+	 * Given a view mode checks if it's the currently selected one 
+	 *
+	 * @param InMode View mode to check
+	 */
+	bool CheckViewMode( EVisualizerViewMode::Type InMode ) const
+	{
+		return (ViewMode == InMode);
+	}
+
+	/** Handles clicking on 'Back to parent' button */
+	FReply OnToParentClicked();
+
+	/** Handles clickong on 'Home' button */
+	FReply OnHomeClicked();
+
+	/** Called when bar graph geometry (size) changes */
+	void OnBarGeometryChanged( FGeometry Geometry );
+
+	/** Gets the currently displayed hierarchy node name */
+	FText GetSelectedCategoryName() const;
+
+	/** Checks if home button should be visible */
+	EVisibility GetHomeButtonVisibility() const;	
+
+	/** Checks if 'to parent' button should be visible */
+	EVisibility GetToParentButtonVisibility() const;
+
+	/** Called when the user clicked bar graph's expand button */
+	FReply ExpandBar( TSharedPtr<FVisualizerEvent> BarGraphEvents );
+
+	/** Checks if the selected event has children with children */
+	bool IsExpandable( TSharedPtr< FVisualizerEvent > InEvent );
+
+	/** Adjusts timeline to match the selected event's start and duration */
+	void AdjustTimeline( TSharedPtr< FVisualizerEvent > InEvent );
+
+	TSharedPtr< FVisualizerEvent > FindSelectedEventsParent( TArray< TSharedPtr< FVisualizerEvent > >& BarGraphs, TSharedPtr< FVisualizerEvent > Selection );
+
+	/** A pointer to the ListView of threads graph bars */
+	TSharedPtr< SListView< TSharedPtr< FVisualizerEvent > > > BarGraphsList;
+
+	/** Currently selected bar graph */
+	TSharedPtr< FVisualizerEvent > SelectedBarGraph;
+
+	/** Original profiler data */
+	TSharedPtr< FVisualizerEvent > ProfileData;
+
+	/** Profiler data view (filtered data) */
+	TArray< TSharedPtr< FVisualizerEvent > > ProfileDataView;
+
+	/** List of all SGraphBar widgets in the tree */
+	TArray< TSharedPtr< SGraphBar > > Graphs;
+
+	/** A pointer to the Zoom Label widget */
+	TSharedPtr< STextBlock > ZoomLabel;
+
+	/** A pointer to the horizontal scrollbar widget */
+	TSharedPtr< SScrollBar > ScrollBar;
+
+	/** A pointer to the horizontal scrollbar widget */
+	TSharedPtr< STimeline > Timeline;
+
+	/** Zoom slider value */
+	float ZoomSliderValue;
+
+	/** Scrollbar offset */
+	float ScrollbarOffset;
+
+	/** Should the OnBarGraphSelectionChangedDelegate be suppressed to avoid event loops */
+	bool bSuppressBarGraphSelectionChangedDelegate;
+
+	/** Delegate used when bar graph selection changes */
+	FOnBarGraphSelectionChanged OnBarGraphSelectionChangedDelegate;
+
+	/** Delegate used when bar graph selection changes */
+	FOnBarGraphExpansionChanged OnBarGraphExpansionChangedDelegate;
+
+	/** Delegate used when single event on the bar graph is selected */
+	FOnBarEventSelectionChanged OnBarEventSelectionChangedDelegate;
+
+	/** Delegate used when the user right-clicks on a bar graph */
+	FOnBarGraphContextMenu OnBarGraphContextMenuDelegate;
+
+	/** Bar visualizer view mode */
+	EVisualizerViewMode::Type ViewMode;
+
+};
diff --git a/Engine/Source/Developer/ProfileVisualizer/Public/SEventsTree.h b/Engine/Source/Developer/ProfileVisualizer/Public/SEventsTree.h
new file mode 100644
--- /dev/null
+++ b/Engine/Source/Developer/ProfileVisualizer/Public/SEventsTree.h
@@ -0,0 +1,284 @@
+// Copyright Epic Games, Inc. All Rights Reserved.
+
+#pragma once
+
+#include "CoreMinimal.h"
+#include "Misc/Attribute.h"
+#include "Widgets/SWidget.h"
+#include "Widgets/DeclarativeSyntaxSupport.h"
+#include "Widgets/SCompoundWidget.h"
+#include "Widgets/Views/STableViewBase.h"
+#include "Widgets/Views/STableRow.h"
+#include "Widgets/Views/STreeView.h"
+#include "VisualizerEvents.h"
+
+/**
+ * Item used in task graph events list
+ */
+class PROFILEVISUALIZER_API SEventItem : public SMultiColumnTableRow< TSharedPtr< FVisualizerEvent > >
+{
+public:
+	
+	SLATE_BEGIN_ARGS( SEventItem )
+		: _EventName()
+		, _EventDuration(0.0)
+		{}
+
+
+		SLATE_ARGUMENT( FString, EventName )
+
+		SLATE_ATTRIBUTE( double, EventDuration )		
+
+	SLATE_END_ARGS()
+	
+	/**
+	 * Generates a widget for task graph events list column
+	 *
+	 * @param InArgs   A declaration from which to construct the widget
+	 */
+	virtual TSharedRef< SWidget > GenerateWidgetForColumn( const FName& ColumnName ) override;
+
+
+	/**
+	 * Construct the widget
+	 *
+	 * @param InArgs   A declaration from which to construct the widget
+	 */
+	void Construct( const FArguments& InArgs, const TSharedRef<STableViewBase>& InOwnerTableView );
+
+private:
+
+	FText GetDurationText() const
+	{
+		static const FNumberFormattingOptions DurationFormatOptions = FNumberFormattingOptions()
+			.SetMinimumFractionalDigits(2)
+			.SetMaximumFractionalDigits(2);
+		return FText::AsNumber(EventDuration.Get(), &DurationFormatOptions);
+	}
+
+	FString EventName;
+	TAttribute< double > EventDuration;
+};
+
+
+/**
+ * Profiler events tree. Contains a tree of profiler event for the selected graph bar.
+ */
+class PROFILEVISUALIZER_API SEventsTree : public SCompoundWidget
+{
+public:
+
+	/** Delegate used when the selection changes */
+	DECLARE_DELEGATE_OneParam( FOnEventSelectionChanged, TSharedPtr< FVisualizerEvent > );
+
+	SLATE_BEGIN_ARGS( SEventsTree )
+		: _ProfileData()
+		, _OnEventSelectionChanged()
+		{}
+
+		/** Profiler data */
+		SLATE_ATTRIBUTE( TSharedPtr< FVisualizerEvent >, ProfileData )
+
+		/** Event for handling selection changes */
+		SLATE_EVENT( FOnEventSelectionChanged, OnEventSelectionChanged )
+
+	SLATE_END_ARGS()
+
+	/**
+	 * Construct the widget
+	 * 
+	 * @param InArgs   A declaration from which to construct the widget
+	 */
+	void Construct( const FArguments& InArgs );
+
+	/**
+	 * Function called when the currently selected bar graph changes
+	 *
+	 * @param Selection Currently selected bar graph
+	 */
+	void HandleBarGraphSelectionChanged( TSharedPtr< FVisualizerEvent > Selection );
+
+	/**
+	 * Function called when the currently expanded bar graph changes
+	 *
+	 * @param Selection Currently expanded bar graph
+	 */
+	void HandleBarGraphExpansionChanged( TSharedPtr< FVisualizerEvent > Selection );
+
+	/**
+	 * Function called when the user selects an event bar in the graph
+	 *
+	 * @param Selection Currently selected event
+	 */
+	void HandleBarEventSelectionChanged( int32 Thread, TSharedPtr< FVisualizerEvent > Selection );
+
+	/** Gets a label for this events tree tab */
+	FString GetTabTitle() const;
+
+	/** Name of 'Name' column */
+	static FName NAME_NameColumn;
+
+	/** Name of 'Duration' column */
+	static FName NAME_DurationColumn;
+
+protected:
+
+	/**
+	 * Function called when the currently selected event in the list of thread events changes
+	 *
+	 * @param Selection Currently selected event
+	 * @param SelectInfo Provides context on how the selection changed
+	 */
+	void OnEventSelectionChanged( TSharedPtr< FVisualizerEvent > Selection, ESelectInfo::Type SelectInfo );
+
+	/** 
+	 * Generates SEventItem widgets for the events tree
+	 *
+	 * @param InItem Event to generate SEventItem for
+	 * @param OwnerTable Owner Table
+	 */
+	TSharedRef< ITableRow > OnGenerateWidgetForEventsList( TSharedPtr< FVisualizerEvent > InItem, const TSharedRef< STableViewBase >& OwnerTable );
+
+	/** 
+	 * Given a profiler event, generates children for it
+	 *
+	 * @param InItem Event to generate children for
+	 * @param OutChildren Generated children
+	 */
+	void OnGetChildrenForEventsList( TSharedPtr<FVisualizerEvent> InItem, TArray<TSharedPtr<FVisualizerEvent> >& OutChildren );
+
+	/** 
+	 * Handles column sorting mode change
+	 *
+	 * @param InSortMode New sorting mode
+	 */
+	void OnColumnSortModeChanged( const EColumnSortPriority::Type SortPriority, const FName& ColumnId, const EColumnSortMode::Type InSortMode );
+
+	/**
+	 * @return Returns the current sort mode of the specified column
+	 */
+	EColumnSortMode::Type GetColumnSortMode( const FName ColumnId );
+
+	/** Sorts the selected events tree */
+	void SortEventsList( void );
+
+	/** Recursively sorts the selected events tree */
+	void SortEventsList( TArray< TSharedPtr< FVisualizerEvent > >& Events );
+
+	/** Restores event selection after changes to the tree */
+	bool RestoreEventSelection( TArray< TSharedPtr< FVisualizerEvent > >& Events );
+
+	/** 
+	 * Sets 'Duration' column time units
+	 *
+	 * @param InUnits New units
+	 */
+	void SetDurationUnits( EVisualizerTimeUnits::Type InUnits );
+
+	/** Recursively counts all events in the tree */
+	int32 CountEvents( TArray< TSharedPtr< FVisualizerEvent > >& Events );
+
+	/** Given the selected events from the bar graph creates a copy of the selection applying the current view and sorting mode */
+	void CreateSelectedEventsView();
+
+	/** Helper function for creating a copy of the selected events in a hierarchy */
+	TSharedPtr< FVisualizerEvent > CreateSelectedEventsViewRecursively( TSharedPtr< FVisualizerEvent > SourceEvent );
+
+	/** Helper function for creating a copy of the selected events and flattening the hierarchy */
+	void CreateSelectedEventsViewRecursivelyAndFlatten( TSharedPtr< FVisualizerEvent > SourceEvent );
+
+	/** Helper function for creating a copy of the selected events combining leaves with the same name */
+	void CreateSelectedEventsViewRecursivelyCoalesced( TArray< TSharedPtr< FVisualizerEvent > >& SourceEvents, TArray< TSharedPtr< FVisualizerEvent > >& CopiedEvents, TSharedPtr< FVisualizerEvent > InParent );
+	
+	/** Helper function for creating a copy of the selected events combining leaves with the same name and flattening the hierarchy */
+	void CreateSelectedEventsViewRecursivelyFlatCoalesced( TArray< TSharedPtr< FVisualizerEvent > >& SourceEvents );
+
+	/** Gets the currently selected time units */
+	bool CheckDurationUnits( EVisualizerTimeUnits::Type InUnits ) const
+	{
+		return (InUnits == DurationUnits);
+	}
+
+	/** 
+	 * Sets the current view mode 
+	 *
+	 * @param InMode New view mode
+	 */
+	void SetViewMode( EVisualizerViewMode::Type InMode );
+
+	/** 
+	 * Given a view mode checks if it's the currently selected one 
+	 *
+	 * @param InMode View mode to check
+	 */
+	bool CheckViewMode( EVisualizerViewMode::Type InMode ) const
+	{
+		return (ViewMode == InMode);
+	}
+
+	/** 
+	 * Converts ms to currently selected time units.
+	 *
+	 * @param InDurationMs duration time in ms.
+	 */
+	double GetEventDuration( double InDurationMs ) const;
+
+	/**
+		* Called by the editable text control when the filter text is changed by the user
+		*
+		* @param	InFilterText	The new text
+		*/
+	void OnFilterTextChanged( const FText& InFilterText );
+
+	/** Called by the editable text control when a user presses enter or commits their text change */
+	void OnFilterTextCommitted( const FText& InFilterText, ETextCommit::Type CommitInfo );
+
+	/** Checks if event name passes current filter */
+	bool FilterEvent( TSharedPtr< FVisualizerEvent > InEvent ) const
+	{
+		// Filter only leaves, we want to keep the hierarchy
+		return InEvent->Children.Num() > 0 || FilterText.Len() == 0 || InEvent->EventName.StartsWith( FilterText );
+	}
+
+	/** Gets the currently selected time units text */
+	FText GetDurationColumnTitle() const;
+
+	/** A pointer to the ListView of profiler events */
+	TSharedPtr< STreeView< TSharedPtr< FVisualizerEvent > > > EventsListView;
+
+	/** Original profiler data */
+	TSharedPtr< FVisualizerEvent > ProfileData;
+
+	/** Currently selected events for this tree */
+	TSharedPtr< FVisualizerEvent > SelectedThread;
+
+	/** List of events for the currently selected thread */
+	TArray< TSharedPtr< FVisualizerEvent > > SelectedEvents;
+
+	/** List of events for the currently selected thread */
+	TArray< TSharedPtr< FVisualizerEvent > > SelectedEventsView;
+
+	/** Maps the events generated using currently selected view and sorting modes to the source events tree */
+	TMap< TSharedPtr< FVisualizerEvent >, TSharedPtr< FVisualizerEvent > > ViewToEventsMap;
+
+	/** Delegate used to notify when event selection changes */
+	FOnEventSelectionChanged OnEventSelectionChangedDelegate;
+
+	/** Specify which column to sort with */
+	FName SortByColumn;
+
+	/** Currently selected sorting mode */
+	EColumnSortMode::Type SortMode;
+
+	/** Currently selected time units */
+	EVisualizerTimeUnits::Type DurationUnits;
+
+	/** Currently selected view mode */
+	EVisualizerViewMode::Type ViewMode;
+
+	/** Events filter text */
+	FString FilterText;
+
+	/** Suppresses SelectionChanged delegate to avoid event loops between graph visualizer and events tree */
+	bool bSuppressSelectionChangedEvent;
+};
diff --git a/Engine/Source/Developer/ProfileVisualizer/Public/SGraphBar.h b/Engine/Source/Developer/ProfileVisualizer/Public/SGraphBar.h
new file mode 100644
--- /dev/null
+++ b/Engine/Source/Developer/ProfileVisualizer/Public/SGraphBar.h
@@ -0,0 +1,153 @@
+// Copyright Epic Games, Inc. All Rights Reserved.
+
+#pragma once
+
+#include "CoreMinimal.h"
+#include "Layout/Geometry.h"
+#include "Input/Reply.h"
+#include "Widgets/DeclarativeSyntaxSupport.h"
+#include "Widgets/SLeafWidget.h"
+#include "Framework/SlateDelegates.h"
+#include "VisualizerEvents.h"
+
+class FPaintArgs;
+class FSlateWindowElementList;
+struct FSlateBrush;
+
+/** A graph bar widget.*/
+class PROFILEVISUALIZER_API SGraphBar : public SLeafWidget
+{
+public:
+
+	/** Delegate used allotted gemetry changes */
+	DECLARE_DELEGATE_OneParam( FOnGeometryChanged, FGeometry );
+
+	typedef TSlateDelegates< TSharedPtr< FVisualizerEvent > >::FOnSelectionChanged FOnSelectionChanged;
+
+public:
+	SLATE_BEGIN_ARGS( SGraphBar )
+		: _OnSelectionChanged()
+		{}
+
+		SLATE_EVENT( FOnSelectionChanged, OnSelectionChanged )
+
+		SLATE_EVENT( FOnGeometryChanged, OnGeometryChanged )
+
+	SLATE_END_ARGS()
+
+	/**
+	 * Construct the widget
+	 * 
+	 * @param InArgs   A declaration from which to construct the widget
+	 */
+	void Construct( const FArguments& InArgs );
+
+	// SWidget interface
+	virtual int32 OnPaint( const FPaintArgs& Args, const FGeometry& AllottedGeometry, const FSlateRect& MyCullingRect, FSlateWindowElementList& OutDrawElements, int32 LayerId, const FWidgetStyle& InWidgetStyle, bool bParentEnabled ) const override;
+	virtual void Tick( const FGeometry& AllottedGeometry, const double InCurrentTime, const float InDeltaTime ) override;
+	virtual FReply OnMouseButtonDown( const FGeometry& MyGeometry, const FPointerEvent& MouseEvent ) override;
+	virtual FReply OnMouseMove( const FGeometry& MyGeometry, const FPointerEvent& MouseEvent ) override;
+	FVector2D ComputeDesiredSize(float) const override;
+	// End of SWidget interface
+
+	/**
+	 * Adds profiler events to draw as bars
+	 *
+	 * @param InEvent	Events to draw
+	 */
+	void SetEvents( const FVisualizerEventsArray& InEvents, double InStartTime, double InTotalTime );
+
+	/**
+	 * Sets the graph's zoom level
+	 *
+	 * @param NewValue	Zoom value, minimum is 1.0
+	 */
+	void SetZoom( float InZoom )
+	{
+		Zoom = FMath::Max(InZoom, 1.0f);
+	}
+
+	/**
+	 * Sets the graph's offset by which all graph bars should be moved
+	 *
+	 * @param NewValue	Offset value
+	 */
+	void SetOffset( float InOffset )
+	{
+		Offset = InOffset;
+	} 
+
+	/**
+	 * Gets the graph's offset value
+	 *
+	 * @return Offset value
+	 */
+	float GetOffset() const
+	{
+		return Offset;
+	}
+
+	/**
+	 * Gets all events this graph will draw
+	 *
+	 * @return List of all events to draw
+	 */
+	FVisualizerEventsArray& GetEvents()
+	{
+		return Events;
+	}
+
+private:
+
+	FORCEINLINE bool CalculateEventGeometry( FVisualizerEvent* InEvent, const FGeometry& InGeometry, float& OutStartX, float& OutEndX ) const
+	{
+		const double SubPixelMinSize = 3.0;
+		const double EventStart = (InEvent->Start - StartTime) / TotalTime;
+		const double EventDuration = InEvent->Duration / TotalTime;
+		const double ClampedStart = FMath::Clamp( Offset + EventStart * Zoom, 0.0, 1.0  );		
+		const double ClampedEnd = FMath::Clamp( Offset + (EventStart + EventDuration) * Zoom, 0.0, 1.0  );
+		const double ClampedSize = ClampedEnd - ClampedStart;
+		OutStartX = (float)( InGeometry.GetLocalSize().X * ClampedStart );
+		OutEndX = OutStartX + (float)FMath::Max( InGeometry.GetLocalSize().X * ClampedSize, ClampedEnd > 0.0f ? SubPixelMinSize : 0.0 );
+		return ClampedEnd > 0.0 && ClampedStart < 1.0;
+	}
+
+	/** Delegate to invoke when selection changes. */
+	FOnSelectionChanged OnSelectionChanged;
+
+	/** Background image to use for the graph bar */
+	const FSlateBrush* BackgroundImage;
+
+	/** Foreground image to use for the graph bar */
+	const FSlateBrush* FillImage;
+
+	/** Image to be used when drawing selected event */
+	const FSlateBrush* SelectedImage;
+
+	/** List of all events to draw */
+	FVisualizerEventsArray Events;
+
+	/** Start time (0.0 - 1.0) */
+	double StartTime;
+	
+	/** End time (0.0 - 1.0) */
+	double TotalTime;
+
+	/** Current zoom of the graph */
+	float Zoom;
+
+	/** Current offset of the graph */
+	float Offset;
+
+	/** Last hovered event index */
+	int32 LastHoveredEvent;
+
+	/** Last allotted geometry */
+	FGeometry LastGeometry;
+
+	/** Delegate called when the geometry changes */
+	FOnGeometryChanged OnGeometryChanged;
+
+	/** Color palette for bars coloring */
+	static FColor	ColorPalette[];
+};
diff --git a/Engine/Source/Developer/ProfileVisualizer/Public/STimeline.h b/Engine/Source/Developer/ProfileVisualizer/Public/STimeline.h
new file mode 100644
--- /dev/null
+++ b/Engine/Source/Developer/ProfileVisualizer/Public/STimeline.h
@@ -0,0 +1,137 @@
+// Copyright Epic Games, Inc. All Rights Reserved.
+
+#pragma once
+
+#include "CoreMinimal.h"
+#include "Layout/Geometry.h"
+#include "Input/Reply.h"
+#include "Widgets/DeclarativeSyntaxSupport.h"
+#include "Widgets/SCompoundWidget.h"
+
+class FPaintArgs;
+class FSlateWindowElementList;
+struct FSlateBrush;
+
+/** A timeline widget.*/
+class PROFILEVISUALIZER_API STimeline : public SCompoundWidget
+{
+
+public:
+	SLATE_BEGIN_ARGS( STimeline )
+		: _MinValue( 0.0f )
+		, _MaxValue( 1.0f )
+		, _FixedLabelSpacing(1.0)
+		{}
+
+		/** Minimum value on the timeline */
+		SLATE_ARGUMENT( float, MinValue )
+
+		/** Maximum value on the timeline */
+		SLATE_ARGUMENT( float, MaxValue )
+
+		/** fixed pixel spacing between centers of labels */
+		SLATE_ARGUMENT( float, FixedLabelSpacing )
+
+	SLATE_END_ARGS()
+
+	/**
+	 * Construct the widget
+	 * 
+	 * @param InArgs   A declaration from which to construct the widget
+	 */
+	void Construct( const FArguments& InArgs );
+
+	// SWidget interface
+	virtual int32 OnPaint( const FPaintArgs& Args,  const FGeometry& AllottedGeometry, const FSlateRect& MyCullingRect, FSlateWindowElementList& OutDrawElements, int32 LayerId, const FWidgetStyle& InWidgetStyle, bool bParentEnabled ) const override;
+	virtual void Tick( const FGeometry& AllottedGeometry, const double InCurrentTime, const float InDeltaTime ) override;
+	virtual FReply OnMouseButtonDown( const FGeometry& MyGeometry, const FPointerEvent& MouseEvent ) override;
+	virtual FReply OnMouseMove( const FGeometry& MyGeometry, const FPointerEvent& MouseEvent ) override;
+	FVector2D ComputeDesiredSize(float) const override;
+	// End of SWidget interface
+
+	/**
+	 * Sets the graph's zoom level
+	 *
+	 * @param NewValue	Zoom value, minimum is 1.0
+	 */
+	void SetZoom( float InZoom )
+	{
+		Zoom = FMath::Max(InZoom, 1.0f);
+	}
+
+	float GetZoom()
+	{
+		return Zoom;
+	}
+	/**
+	 * Sets the graph's offset by which all graph bars should be moved
+	 *
+	 * @param NewValue	Offset value
+	 */
+	void SetOffset( float InOffset )
+	{
+		Offset = InOffset;
+	} 
+
+	/**
+	 * Gets the graph's offset value
+	 *
+	 * @return Offset value
+	 */
+	float GetOffset() const
+	{
+		return Offset;
+	}
+
+	/**
+	 * Sets the graph's min and max values.
+	 *
+	 * @param InMin	New min value.
+	 * @param InMin	New max value.
+	 */
+	void SetMinMaxValues(float InMin, float InMax)
+	{
+		MinValue = InMin;
+		MaxValue = InMax;
+	}
+
+	void GetMinMaxValues(float &InMin, float &InMax)
+	{
+		InMin = MinValue;
+		InMax = MaxValue;
+	}
+
+	void SetDrawingGeometry(const FGeometry& Geometry)
+	{
+		DrawingGeometry = Geometry;
+	}
+
+	FGeometry GetDrawingGeometry() const
+	{
+		return DrawingGeometry;
+	}
+
+private:
+
+	/** Background image to use for the graph bar */
+	const FSlateBrush* BackgroundImage;
+
+	/** Minimum value on the timeline */
+	float MinValue;
+
+	/** Maximum value on the timeline */
+	float MaxValue;
+
+	/** fixed pixel spacing between centers of labels */
+	float FixedLabelSpacing;
+
+	/** Current zoom of the graph */
+	float Zoom;
+
+	/** Current offset of the graph */
+	float Offset;
+
+	float DrawingOffsetX;
+
+	FGeometry DrawingGeometry;
+};
diff --git a/Engine/Source/Developer/ProfileVisualizer/Public/VisualizerEvents.h b/Engine/Source/Developer/ProfileVisualizer/Public/VisualizerEvents.h
new file mode 100644
--- /dev/null
+++ b/Engine/Source/Developer/ProfileVisualizer/Public/VisualizerEvents.h
@@ -0,0 +1,96 @@
+// Copyright Epic Games, Inc. All Rights Reserved.
+
+#pragma once
+
+#include "Containers/Array.h"
+#include "Containers/UnrealString.h"
+#include "CoreMinimal.h"
+#include "HAL/Platform.h"
+#include "HAL/PlatformCrt.h"
+#include "Templates/SharedPointer.h"
+
+class FArchive;
+
+namespace EVisualizerTimeUnits
+{
+	typedef uint8 Type;
+
+		/** Microseconds */
+	const Type Microseconds = 0;
+
+		/** Milliseconds */
+	const Type Milliseconds = 1;
+
+		/** Seconds */
+	const Type Seconds = 2;
+};
+
+namespace EVisualizerViewMode
+{
+	typedef uint8 Type;
+
+	/** Hierarchical */
+	const Type Hierarchical = 0;
+
+	/** Flat */
+	const Type Flat = 1;
+
+	/** Coalesced */
+	const Type Coalesced = 2;
+
+	/** FlatCoalesced */
+	const Type FlatCoalesced = 3;
+};
+
+/** A graph event represented by SGraphBar as a single bar.*/
+struct FVisualizerEvent
+{	
+	/** Normalized start time (0.0-1.0) of the event relative to the first event in the profile */
+	double Start;
+
+	/** Normalized duration time (0.0-1.0) of the event. */
+	double Duration;
+
+	/** Duration of the event in milliseconds */
+	double DurationMs;
+
+	/** Category this event belongs to (thread/file etc.) */
+	int32 Category;
+
+	/** Name of the event. */
+	FString EventName;
+
+	/** Determines if this event is selected or not. */
+	bool IsSelected;
+
+	/** Bar color */
+	uint32 ColorIndex;
+
+	/** Parent event */
+	TSharedPtr< FVisualizerEvent > ParentEvent;
+
+	/**  Child events */
+	TArray< TSharedPtr< FVisualizerEvent > > Children;
+
+	FVisualizerEvent(const double InStart, const double InDuration, const double InDurationMs, const int32 InCategory, const FString& InEventName)
+		: Start( InStart )
+		, Duration( InDuration )
+		, DurationMs( InDurationMs )
+		, Category( InCategory )
+		, EventName( InEventName )
+		, IsSelected( false )
+		, ColorIndex( 0 )
+	{
+		ColorIndex = GetTypeHash( InEventName );
+	}
+
+	static TSharedPtr< FVisualizerEvent > LoadVisualizerEvent(FArchive *Ar);
+	static void SaveVisualizerEventRecursively(FArchive *Ar, TSharedPtr< FVisualizerEvent > VisualizerEvent);
+
+private:
+	static TSharedPtr< FVisualizerEvent > LoadVisualizerEventRecursively(FArchive *Ar, TSharedPtr< FVisualizerEvent > InParentEvent);
+
+};
+
+/** Array of graph events.*/
+typedef TArray< TSharedPtr< FVisualizerEvent > > FVisualizerEventsArray;
diff --git a/Engine/Source/Developer/RealtimeProfiler/Private/SRealtimeProfilerFrame.cpp b/Engine/Source/Developer/RealtimeProfiler/Private/SRealtimeProfilerFrame.cpp
--- a/Engine/Source/Developer/RealtimeProfiler/Private/SRealtimeProfilerFrame.cpp
+++ b/Engine/Source/Developer/RealtimeProfiler/Private/SRealtimeProfilerFrame.cpp
@@ -55,7 +55,7 @@ void SRealtimeProfilerFrame::AddVisualizer( TSharedPtr< SDockTabStack > TabStack
 
 	Visualizers.Add( NewVisualizer );
 
-	AddTab( NewVisualizer, "Stats", NSLOCTEXT("TaskGraph", "ProfileVisualizerToolTip", "Profile Visualizer.").ToString(), TabStack );
+	AddTab( NewVisualizer, "Stats", NSLOCTEXT("ProfileVisualizer", "ProfileVisualizerToolTip", "Profile Visualizer.").ToString(), TabStack );
 
 
 	Visualizer = NewVisualizer;
diff --git a/Engine/Source/Developer/RealtimeProfiler/Private/SRealtimeProfilerTimeline.cpp b/Engine/Source/Developer/RealtimeProfiler/Private/SRealtimeProfilerTimeline.cpp
--- a/Engine/Source/Developer/RealtimeProfiler/Private/SRealtimeProfilerTimeline.cpp
+++ b/Engine/Source/Developer/RealtimeProfiler/Private/SRealtimeProfilerTimeline.cpp
@@ -106,7 +106,7 @@ FText SRealtimeProfilerTimeline::GetZoomLabel() const
 	static const FNumberFormattingOptions ZoomFormatOptions = FNumberFormattingOptions()
 		.SetMinimumFractionalDigits(2)
 		.SetMaximumFractionalDigits(2);
-	return FText::Format( NSLOCTEXT("TaskGraph", "ZoomLabelFmt", "Zoom: {0}x"), FText::AsNumber(GetZoom(), &ZoomFormatOptions) );
+	return FText::Format( NSLOCTEXT("ProfileVisualizer", "ZoomLabelFmt", "Zoom: {0}x"), FText::AsNumber(GetZoom(), &ZoomFormatOptions) );
 }
 
 float SRealtimeProfilerTimeline::GetZoomValue() const
diff --git a/Engine/Source/Developer/RealtimeProfiler/Private/SRealtimeProfilerVisualizer.cpp b/Engine/Source/Developer/RealtimeProfiler/Private/SRealtimeProfilerVisualizer.cpp
--- a/Engine/Source/Developer/RealtimeProfiler/Private/SRealtimeProfilerVisualizer.cpp
+++ b/Engine/Source/Developer/RealtimeProfiler/Private/SRealtimeProfilerVisualizer.cpp
@@ -114,7 +114,7 @@ TSharedRef<SWidget> SRealtimeProfilerVisualizer::MakeBarVisualizerContextMenu()
 	FMenuBuilder MenuBuilder( bCloseAfterSelection, NULL );
 	{
 		FUIAction Action( FExecuteAction::CreateSP( this, &SRealtimeProfilerVisualizer::ShowGraphBarInEventsWindow, (int32)INDEX_NONE ) );
-		MenuBuilder.AddMenuEntry( NSLOCTEXT("TaskGraph", "GraphBarShowInNew", "Show in New Events Window"), FText::GetEmpty(), FSlateIcon(), Action, NAME_None, EUserInterfaceActionType::Button );
+		MenuBuilder.AddMenuEntry( NSLOCTEXT("ProfileVisualizer", "GraphBarShowInNew", "Show in New Events Window"), FText::GetEmpty(), FSlateIcon(), Action, NAME_None, EUserInterfaceActionType::Button );
 	}
 		
 	return MenuBuilder.MakeWidget();
diff --git a/Engine/Source/Developer/RealtimeProfiler/RealtimeProfiler.Build.cs b/Engine/Source/Developer/RealtimeProfiler/RealtimeProfiler.Build.cs
--- a/Engine/Source/Developer/RealtimeProfiler/RealtimeProfiler.Build.cs
+++ b/Engine/Source/Developer/RealtimeProfiler/RealtimeProfiler.Build.cs
@@ -14,7 +14,7 @@ public class RealtimeProfiler : ModuleRules
 				"Slate",
 				"SlateCore",
                 
-				"TaskGraph",
+				"ProfileVisualizer",
 				"Engine"
 			}
 		);
diff --git a/Engine/Source/Developer/TaskGraph/Private/SBarVisualizer.cpp b/Engine/Source/Developer/TaskGraph/Private/SBarVisualizer.cpp
deleted file mode 100644
diff --git a/Engine/Source/Developer/TaskGraph/Private/SEventsTree.cpp b/Engine/Source/Developer/TaskGraph/Private/SEventsTree.cpp
deleted file mode 100644
diff --git a/Engine/Source/Developer/TaskGraph/Private/SGraphBar.cpp b/Engine/Source/Developer/TaskGraph/Private/SGraphBar.cpp
deleted file mode 100644
diff --git a/Engine/Source/Developer/TaskGraph/Private/SProfileVisualizer.cpp b/Engine/Source/Developer/TaskGraph/Private/SProfileVisualizer.cpp
deleted file mode 100644
diff --git a/Engine/Source/Developer/TaskGraph/Private/SProfileVisualizer.h b/Engine/Source/Developer/TaskGraph/Private/SProfileVisualizer.h
deleted file mode 100644
diff --git a/Engine/Source/Developer/TaskGraph/Private/STaskGraph.cpp b/Engine/Source/Developer/TaskGraph/Private/STaskGraph.cpp
deleted file mode 100644
diff --git a/Engine/Source/Developer/TaskGraph/Private/STimeline.cpp b/Engine/Source/Developer/TaskGraph/Private/STimeline.cpp
deleted file mode 100644
diff --git a/Engine/Source/Developer/TaskGraph/Private/TaskGraphStyle.cpp b/Engine/Source/Developer/TaskGraph/Private/TaskGraphStyle.cpp
deleted file mode 100644
diff --git a/Engine/Source/Developer/TaskGraph/Private/TaskGraphStyle.h b/Engine/Source/Developer/TaskGraph/Private/TaskGraphStyle.h
deleted file mode 100644
diff --git a/Engine/Source/Developer/TaskGraph/Private/VisualizerEvents.cpp b/Engine/Source/Developer/TaskGraph/Private/VisualizerEvents.cpp
deleted file mode 100644
diff --git a/Engine/Source/Developer/TaskGraph/Public/SBarVisualizer.h b/Engine/Source/Developer/TaskGraph/Public/SBarVisualizer.h
deleted file mode 100644
diff --git a/Engine/Source/Developer/TaskGraph/Public/SEventsTree.h b/Engine/Source/Developer/TaskGraph/Public/SEventsTree.h
deleted file mode 100644
diff --git a/Engine/Source/Developer/TaskGraph/Public/SGraphBar.h b/Engine/Source/Developer/TaskGraph/Public/SGraphBar.h
deleted file mode 100644
diff --git a/Engine/Source/Developer/TaskGraph/Public/STaskGraph.h b/Engine/Source/Developer/TaskGraph/Public/STaskGraph.h
deleted file mode 100644
diff --git a/Engine/Source/Developer/TaskGraph/Public/STimeline.h b/Engine/Source/Developer/TaskGraph/Public/STimeline.h
deleted file mode 100644
diff --git a/Engine/Source/Developer/TaskGraph/Public/VisualizerEvents.h b/Engine/Source/Developer/TaskGraph/Public/VisualizerEvents.h
deleted file mode 100644
diff --git a/Engine/Source/Developer/TaskGraph/TaskGraph.Build.cs b/Engine/Source/Developer/TaskGraph/TaskGraph.Build.cs
deleted file mode 100644
diff --git a/Engine/Source/Programs/ReplicationSystemTest/ReplicationSystemTest.Build.cs b/Engine/Source/Programs/ReplicationSystemTest/ReplicationSystemTest.Build.cs
--- a/Engine/Source/Programs/ReplicationSystemTest/ReplicationSystemTest.Build.cs
+++ b/Engine/Source/Programs/ReplicationSystemTest/ReplicationSystemTest.Build.cs
@@ -39,7 +39,7 @@ public class ReplicationSystemTest : ModuleRules
 				"SessionServices",
 				"SlateNullRenderer",
 				"SlateRHIRenderer",
-				"TaskGraph",
+				"ProfileVisualizer",
  			}
 		);
 	}
diff --git a/Engine/Source/Runtime/D3D12RHI/D3D12RHI.Build.cs b/Engine/Source/Runtime/D3D12RHI/D3D12RHI.Build.cs
--- a/Engine/Source/Runtime/D3D12RHI/D3D12RHI.Build.cs
+++ b/Engine/Source/Runtime/D3D12RHI/D3D12RHI.Build.cs
@@ -27,11 +27,6 @@ public class D3D12RHI : ModuleRules
 			}
 			);
 
-		if (Target.Configuration != UnrealTargetConfiguration.Shipping)
-		{
-			PrivateIncludePathModuleNames.AddRange(new string[] { "TaskGraph" });
-		}
-
 		///////////////////////////////////////////////////////////////
         // Platform specific defines
         ///////////////////////////////////////////////////////////////
diff --git a/Engine/Source/Runtime/D3D12RHI/Private/D3D12RHI.cpp b/Engine/Source/Runtime/D3D12RHI/Private/D3D12RHI.cpp
--- a/Engine/Source/Runtime/D3D12RHI/Private/D3D12RHI.cpp
+++ b/Engine/Source/Runtime/D3D12RHI/Private/D3D12RHI.cpp
@@ -14,10 +14,6 @@
 #include "Windows/HideWindowsPlatformTypes.h"
 #endif
 
-#if !UE_BUILD_SHIPPING
-#include "STaskGraph.h"
-#endif
-
 #if !defined(D3D12_PLATFORM_NEEDS_DISPLAY_MODE_ENUMERATION)
 	#define D3D12_PLATFORM_NEEDS_DISPLAY_MODE_ENUMERATION 1
 #endif
diff --git a/Engine/Source/Runtime/Engine/Engine.Build.cs b/Engine/Source/Runtime/Engine/Engine.Build.cs
--- a/Engine/Source/Runtime/Engine/Engine.Build.cs
+++ b/Engine/Source/Runtime/Engine/Engine.Build.cs
@@ -57,7 +57,6 @@ public class Engine : ModuleRules
 		{
 			PrivateIncludePathModuleNames.AddRange(
 				new string[] {
-					"TaskGraph",
 					"SlateReflector",
 				}
 			);
diff --git a/Engine/Source/Runtime/Engine/Private/LevelTick.cpp b/Engine/Source/Runtime/Engine/Private/LevelTick.cpp
--- a/Engine/Source/Runtime/Engine/Private/LevelTick.cpp
+++ b/Engine/Source/Runtime/Engine/Private/LevelTick.cpp
@@ -54,10 +54,6 @@
 //#include "SoundDefinitions.h"
 #include "FXSystem.h"
 #include "TickTaskManagerInterface.h"
-#if !UE_BUILD_SHIPPING
-#include "VisualizerEvents.h"
-#include "STaskGraph.h"
-#endif
 #include "Async/ParallelFor.h"
 #include "Engine/CoreSettings.h"
 
diff --git a/Engine/Source/Runtime/Launch/Launch.Build.cs b/Engine/Source/Runtime/Launch/Launch.Build.cs
--- a/Engine/Source/Runtime/Launch/Launch.Build.cs
+++ b/Engine/Source/Runtime/Launch/Launch.Build.cs
@@ -9,7 +9,7 @@ public class Launch : ModuleRules
 	{
 		PrivateIncludePathModuleNames.AddRange(new string[] {
 				"AutomationController",
-				"TaskGraph",
+				"ProfileVisualizer",
 		});
 
 		PrivateDependencyModuleNames.AddRange(new string[] {
@@ -167,7 +167,7 @@ public class Launch : ModuleRules
 			PublicIncludePathModuleNames.Add("ProfilerService");
 
 			DynamicallyLoadedModuleNames.AddRange(new string[] {
-				"TaskGraph",
+				"ProfileVisualizer",
 				"RealtimeProfiler",
 				"ProfilerService"
 			});
diff --git a/Engine/Source/Runtime/Launch/Private/LaunchEngineLoop.cpp b/Engine/Source/Runtime/Launch/Private/LaunchEngineLoop.cpp
--- a/Engine/Source/Runtime/Launch/Private/LaunchEngineLoop.cpp
+++ b/Engine/Source/Runtime/Launch/Private/LaunchEngineLoop.cpp
@@ -199,7 +199,7 @@
 	#include "ShaderPipelineCache.h"
 
 #if !UE_BUILD_SHIPPING
-	#include "STaskGraph.h"
+	#include "ProfileVisualizerModule.h"
 	#include "IProfilerServiceModule.h"
 #endif
 
@@ -3933,7 +3933,7 @@ int32 FEngineLoop::PreInitPostStartupScreen(const TCHAR* CmdLine)
 
 	// Initialize profile visualizers.
 #if !(UE_BUILD_SHIPPING || UE_BUILD_TEST)
-	FModuleManager::Get().LoadModule(TEXT("TaskGraph"));
+	FModuleManager::Get().LoadModule(TEXT("ProfileVisualizer"));
 	if (FPlatformProcess::SupportsMultithreading())
 	{
 		FModuleManager::Get().LoadModule(TEXT("ProfilerService"));
diff --git a/Engine/Source/Runtime/OpenGLDrv/OpenGLDrv.Build.cs b/Engine/Source/Runtime/OpenGLDrv/OpenGLDrv.Build.cs
--- a/Engine/Source/Runtime/OpenGLDrv/OpenGLDrv.Build.cs
+++ b/Engine/Source/Runtime/OpenGLDrv/OpenGLDrv.Build.cs
@@ -32,11 +32,6 @@ public class OpenGLDrv : ModuleRules
 			AddEngineThirdPartyPrivateStaticDependencies(Target, "SDL2");
 		}
 
-		if (Target.Configuration != UnrealTargetConfiguration.Shipping)
-		{
-			PrivateIncludePathModuleNames.Add("TaskGraph");
-		}
-
 		if (Target.Platform == UnrealTargetPlatform.Android)
 		{
 			PrivateDependencyModuleNames.Add("detex");
diff --git a/Engine/Source/Runtime/RHI/Private/GPUProfiler.cpp b/Engine/Source/Runtime/RHI/Private/GPUProfiler.cpp
--- a/Engine/Source/Runtime/RHI/Private/GPUProfiler.cpp
+++ b/Engine/Source/Runtime/RHI/Private/GPUProfiler.cpp
@@ -9,7 +9,7 @@
 
 #if !UE_BUILD_SHIPPING
 #include "VisualizerEvents.h"
-#include "STaskGraph.h"
+#include "ProfileVisualizerModule.h"
 #include "Modules/ModuleManager.h"
 #endif
 
@@ -580,10 +580,10 @@ void FGPUProfilerEventNodeFrame::DumpEventTree()
 				{
 					void Thread( TSharedPtr<FVisualizerEvent> InVisualizerData, const FText InVsyncEnabledWarningText )
 					{
-						static FName TaskGraphModule(TEXT("TaskGraph"));			
-						if (FModuleManager::Get().IsModuleLoaded(TaskGraphModule))
+						static FName ProfileVisualizerModule(TEXT("ProfileVisualizer"));			
+						if (FModuleManager::Get().IsModuleLoaded(ProfileVisualizerModule))
 						{
-							IProfileVisualizerModule& ProfileVisualizer = FModuleManager::GetModuleChecked<IProfileVisualizerModule>(TaskGraphModule);
+							IProfileVisualizerModule& ProfileVisualizer = FModuleManager::GetModuleChecked<IProfileVisualizerModule>(ProfileVisualizerModule);
 							// Display a warning if this is a GPU profile and the GPU was profiled with v-sync enabled (otherwise InVsyncEnabledWarningText is empty)
 							ProfileVisualizer.DisplayProfileVisualizer( InVisualizerData, TEXT("GPU"), InVsyncEnabledWarningText, FLinearColor::Red );
 						}
diff --git a/Engine/Source/Runtime/RHI/RHI.Build.cs b/Engine/Source/Runtime/RHI/RHI.Build.cs
--- a/Engine/Source/Runtime/RHI/RHI.Build.cs
+++ b/Engine/Source/Runtime/RHI/RHI.Build.cs
@@ -55,7 +55,7 @@ public class RHI : ModuleRules
 
 		if (Target.Configuration != UnrealTargetConfiguration.Shipping)
 		{
-			PrivateIncludePathModuleNames.AddRange(new string[] { "TaskGraph" });
+			PrivateIncludePathModuleNames.AddRange(new string[] { "ProfileVisualizer" });
 		}
     }
 }
diff --git a/Engine/Source/Runtime/Windows/D3D11RHI/D3D11RHI.Build.cs b/Engine/Source/Runtime/Windows/D3D11RHI/D3D11RHI.Build.cs
--- a/Engine/Source/Runtime/Windows/D3D11RHI/D3D11RHI.Build.cs
+++ b/Engine/Source/Runtime/Windows/D3D11RHI/D3D11RHI.Build.cs
@@ -41,10 +41,5 @@ public class D3D11RHI : ModuleRules
 		{
 			PrivateIncludePaths.Add(Path.Combine(this.ModuleDirectory, "Private", "Windows"));
 		}
-
-        if (Target.Configuration != UnrealTargetConfiguration.Shipping)
-		{
-			PrivateIncludePathModuleNames.AddRange(new string[] { "TaskGraph" });
-		}
 	}
 }
diff --git a/Engine/Source/Runtime/Windows/D3D11RHI/Private/D3D11RHI.cpp b/Engine/Source/Runtime/Windows/D3D11RHI/Private/D3D11RHI.cpp
--- a/Engine/Source/Runtime/Windows/D3D11RHI/Private/D3D11RHI.cpp
+++ b/Engine/Source/Runtime/Windows/D3D11RHI/Private/D3D11RHI.cpp
@@ -18,10 +18,6 @@
 #endif	//WITH_DX_PERF
 #include "OneColorShader.h"
 
-#if !UE_BUILD_SHIPPING
-	#include "STaskGraph.h"
-#endif
-
 DEFINE_LOG_CATEGORY(LogD3D11RHI);
 
 extern void UniformBufferBeginFrame();
