Update ProfileVisualizer.Build.cs to be more broad in when it adds certain module dependencies, to support inclusion in programs compiling against the editor.

diff --git a/Engine/Source/Developer/ProfileVisualizer/ProfileVisualizer.Build.cs b/Engine/Source/Developer/ProfileVisualizer/ProfileVisualizer.Build.cs
--- a/Engine/Source/Developer/ProfileVisualizer/ProfileVisualizer.Build.cs
+++ b/Engine/Source/Developer/ProfileVisualizer/ProfileVisualizer.Build.cs
@@ -17,7 +17,7 @@ public class ProfileVisualizer : ModuleRules
 			}
 		);
 
-		if (Target.Type == TargetType.Editor)
+		if ((Target.bCompileAgainstEditor && (Target.Type == TargetType.Editor || Target.Type == TargetType.Program)))
 		{
 			PrivateDependencyModuleNames.AddRange(
 				new string[] {
