Making FDelayedVisualizerSpawner transactionally safe (this is being used in Lyra DedicatedServer).

diff --git a/Engine/Source/Developer/ProfileVisualizer/Private/ProfileVisualizerModule.cpp b/Engine/Source/Developer/ProfileVisualizer/Private/ProfileVisualizerModule.cpp
--- a/Engine/Source/Developer/ProfileVisualizer/Private/ProfileVisualizerModule.cpp
+++ b/Engine/Source/Developer/ProfileVisualizer/Private/ProfileVisualizerModule.cpp
@@ -77,18 +77,16 @@ public:
 
 	void AddPendingData( TSharedPtr<FVisualizerEvent> InProfileData, const FText& InTitle, const FText& InType )
 	{
-		DataLock.Lock();
+		UE::TScopeLock ScopeLock(DataLock);
 
 		VisualizerDataToSpawn.Add( MakeShareable( new FPendingWindow( InProfileData, InTitle, InType ) )  );
-
-		DataLock.Unlock();
 	}
 
 	// FTickableGameObject interface
 
 	void Tick(float DeltaTime) override
 	{
-		DataLock.Lock();
+		UE::TScopeLock ScopeLock(DataLock);
 
 		for ( int32 Index = 0; Index < VisualizerDataToSpawn.Num(); Index++ )
 		{
@@ -97,8 +95,6 @@ public:
 		}
 
 		VisualizerDataToSpawn.Empty();
-
-		DataLock.Unlock();
 	}
 
 	/** We should call Tick on this object */
@@ -120,7 +116,7 @@ public:
 
 private:
 
-	FCriticalSection DataLock;
+	FTransactionallySafeCriticalSection DataLock;
 
 	/** List of profile data sets to spawn visualizers for */
 	TArray< TSharedPtr<FPendingWindow> > VisualizerDataToSpawn;
